Fix for IOS issues in gameplay tags

Moved log declaration to after Core and added parentesis around !check to satisfy the clang compiler.

diff --git a/Engine/Source/Runtime/GameplayTags/Private/GameplayTagContainer.cpp b/Engine/Source/Runtime/GameplayTags/Private/GameplayTagContainer.cpp
--- a/Engine/Source/Runtime/GameplayTags/Private/GameplayTagContainer.cpp
+++ b/Engine/Source/Runtime/GameplayTags/Private/GameplayTagContainer.cpp
@@ -38,7 +38,7 @@ bool FGameplayTagContainer::operator!=(FGameplayTagContainer const& Other) const
 	{
 		return true;
 	}
-	return !Filter(Other, EGameplayTagMatchType::Explicit, EGameplayTagMatchType::Explicit, EGameplayContainerMatchType::All).Num() == this->Num();
+	return !(Filter(Other, EGameplayTagMatchType::Explicit, EGameplayTagMatchType::Explicit, EGameplayContainerMatchType::All).Num() == this->Num());
 }
 
 bool FGameplayTagContainer::HasTag(FGameplayTag const& TagToCheck, TEnumAsByte<EGameplayTagMatchType::Type> TagMatchType, TEnumAsByte<EGameplayTagMatchType::Type> TagToCheckMatchType) const
