Changes to gameplay tag container matching.

 - Modified the filter function to be independent of the container match type.
 - MatchesAll and MatchesAny now have their own code path that is slightly more efficient for them.

diff --git a/Engine/Source/Runtime/GameplayTags/Classes/GameplayTagContainer.h b/Engine/Source/Runtime/GameplayTags/Classes/GameplayTagContainer.h
--- a/Engine/Source/Runtime/GameplayTags/Classes/GameplayTagContainer.h
+++ b/Engine/Source/Runtime/GameplayTags/Classes/GameplayTagContainer.h
@@ -122,11 +122,10 @@ struct GAMEPLAYTAGS_API FGameplayTagContainer
 	 * @param OtherContainer		The Container to filter against
 	 * @param TagMatchType			Type of match to use for the tags in this container
 	 * @param OtherTagMatchType		Type of match to use for the tags in the OtherContainer param
-	 * @param ContainerMatchType	Type of match to use for filtering
 	 *
 	 * @return A FGameplayTagContainer containing the filtered tags
 	 */
-	FGameplayTagContainer Filter(const FGameplayTagContainer& OtherContainer, TEnumAsByte<EGameplayTagMatchType::Type> TagMatchType, TEnumAsByte<EGameplayTagMatchType::Type> OtherTagMatchType, TEnumAsByte<EGameplayContainerMatchType::Type> ContainerMatchType) const;
+	FGameplayTagContainer Filter(const FGameplayTagContainer& OtherContainer, TEnumAsByte<EGameplayTagMatchType::Type> TagMatchType, TEnumAsByte<EGameplayTagMatchType::Type> OtherTagMatchType) const;
 
 	/**
 	 * Checks if this container matches ANY of the tags in the specified container. Performs matching by expanding this container out
@@ -214,6 +213,20 @@ struct GAMEPLAYTAGS_API FGameplayTagContainer
 		return GameplayTags.CreateConstIterator();
 	}
 
+protected:
+	/**
+	* Returns true if the tags in this container match the tags in OtherContainer for the specified matching types.
+	*
+	* @param OtherContainer		The Container to filter against
+	* @param TagMatchType			Type of match to use for the tags in this container
+	* @param OtherTagMatchType		Type of match to use for the tags in the OtherContainer param
+	* @param ContainerMatchType	Type of match to use for filtering
+	*
+	* @return Returns true if ContainerMatchType is Any and any of the tags in OtherContainer match the tags in this or ContainerMatchType is All and all of the tags in OtherContainer match at least one tag in this. Returns false otherwise.
+	*/
+	bool DoesTagContainerMatch(const FGameplayTagContainer& OtherContainer, TEnumAsByte<EGameplayTagMatchType::Type> TagMatchType, TEnumAsByte<EGameplayTagMatchType::Type> OtherTagMatchType, TEnumAsByte<EGameplayContainerMatchType::Type> ContainerMatchType) const;
+
+
 private:
 
 	// @todo: Remove this via engine version bump/native serialization soon
diff --git a/Engine/Source/Runtime/GameplayTags/Private/GameplayTagContainer.cpp b/Engine/Source/Runtime/GameplayTags/Private/GameplayTagContainer.cpp
--- a/Engine/Source/Runtime/GameplayTags/Private/GameplayTagContainer.cpp
+++ b/Engine/Source/Runtime/GameplayTags/Private/GameplayTagContainer.cpp
@@ -29,7 +29,7 @@ bool FGameplayTagContainer::operator==(FGameplayTagContainer const& Other) const
 	{
 		return false;
 	}
-	return Filter(Other, EGameplayTagMatchType::Explicit, EGameplayTagMatchType::Explicit, EGameplayContainerMatchType::All).Num() == this->Num();
+	return Filter(Other, EGameplayTagMatchType::Explicit, EGameplayTagMatchType::Explicit).Num() == this->Num();
 }
 
 bool FGameplayTagContainer::operator!=(FGameplayTagContainer const& Other) const
@@ -38,7 +38,7 @@ bool FGameplayTagContainer::operator!=(FGameplayTagContainer const& Other) const
 	{
 		return true;
 	}
-	return Filter(Other, EGameplayTagMatchType::Explicit, EGameplayTagMatchType::Explicit, EGameplayContainerMatchType::All).Num() != this->Num();
+	return Filter(Other, EGameplayTagMatchType::Explicit, EGameplayTagMatchType::Explicit).Num() != this->Num();
 }
 
 bool FGameplayTagContainer::HasTag(FGameplayTag const& TagToCheck, TEnumAsByte<EGameplayTagMatchType::Type> TagMatchType, TEnumAsByte<EGameplayTagMatchType::Type> TagToCheckMatchType) const
@@ -68,38 +68,59 @@ FGameplayTagContainer FGameplayTagContainer::GetGameplayTagParents() const
 	return ResultContainer;
 }
 
-FGameplayTagContainer FGameplayTagContainer::Filter(const FGameplayTagContainer& OtherContainer, TEnumAsByte<EGameplayTagMatchType::Type> TagMatchType, TEnumAsByte<EGameplayTagMatchType::Type> OtherTagMatchType, TEnumAsByte<EGameplayContainerMatchType::Type> ContainerMatchType) const
+FGameplayTagContainer FGameplayTagContainer::Filter(const FGameplayTagContainer& OtherContainer, TEnumAsByte<EGameplayTagMatchType::Type> TagMatchType, TEnumAsByte<EGameplayTagMatchType::Type> OtherTagMatchType) const
 {
 	FGameplayTagContainer ResultContainer;
 	UGameplayTagsManager& TagManager = IGameplayTagsModule::Get().GetGameplayTagsManager();
 
 	
+	for (TArray<FGameplayTag>::TConstIterator OtherIt(OtherContainer.GameplayTags); OtherIt; ++OtherIt)
+	{
+		for (TArray<FGameplayTag>::TConstIterator It(this->GameplayTags); It; ++It)
+		{
+			if (TagManager.GameplayTagsMatch(*It, *OtherIt, TagMatchType, OtherTagMatchType) == true)
+			{
+				ResultContainer.AddTag(*It);
+			}
+		}
+	}
+
+	return ResultContainer;
+}
+
+bool FGameplayTagContainer::DoesTagContainerMatch(const FGameplayTagContainer& OtherContainer, TEnumAsByte<EGameplayTagMatchType::Type> TagMatchType, TEnumAsByte<EGameplayTagMatchType::Type> OtherTagMatchType, TEnumAsByte<EGameplayContainerMatchType::Type> ContainerMatchType) const
+{
+	UGameplayTagsManager& TagManager = IGameplayTagsModule::Get().GetGameplayTagsManager();
+
 	for (TArray<FGameplayTag>::TConstIterator OtherIt(OtherContainer.GameplayTags); OtherIt; ++OtherIt)
 	{
 		bool bTagFound = false;
-		
+
 		for (TArray<FGameplayTag>::TConstIterator It(this->GameplayTags); It; ++It)
 		{
 			if (TagManager.GameplayTagsMatch(*It, *OtherIt, TagMatchType, OtherTagMatchType) == true)
 			{
-				bTagFound = true;
 				if (ContainerMatchType == EGameplayContainerMatchType::Any)
 				{
-					ResultContainer.AddTag(*It);
+					return true;
 				}
+
+				bTagFound = true;
+
+				// we only need one match per tag in OtherContainer, so don't bother looking for more
+				break;
 			}
 		}
+
 		if (ContainerMatchType == EGameplayContainerMatchType::All && bTagFound == false)
 		{
-			return ResultContainer; // return the empty container;
+			return false;
 		}
 	}
 
-	if (ContainerMatchType == EGameplayContainerMatchType::All)
-	{
-		ResultContainer.AppendTags(*this);
-	}
-	return ResultContainer;
+	// if we've reached this far then either we are looking for any match and didn't find one (return false) or we're looking for all matches and didn't miss one (return true).
+	check(ContainerMatchType == EGameplayContainerMatchType::All || ContainerMatchType == EGameplayContainerMatchType::Any);
+	return ContainerMatchType == EGameplayContainerMatchType::All;
 }
 
 
@@ -110,7 +131,7 @@ bool FGameplayTagContainer::MatchesAll(FGameplayTagContainer const& Other, bool
 		return bCountEmptyAsMatch;
 	}
 
-	return Filter(Other, EGameplayTagMatchType::IncludeParentTags, EGameplayTagMatchType::Explicit, EGameplayContainerMatchType::All).Num() > 0;
+	return DoesTagContainerMatch(Other, EGameplayTagMatchType::IncludeParentTags, EGameplayTagMatchType::Explicit, EGameplayContainerMatchType::All);
 }
 
 bool FGameplayTagContainer::MatchesAny(FGameplayTagContainer const& Other, bool bCountEmptyAsMatch) const
@@ -120,7 +141,7 @@ bool FGameplayTagContainer::MatchesAny(FGameplayTagContainer const& Other, bool
 		return bCountEmptyAsMatch;
 	}
 
-	return Filter(Other, EGameplayTagMatchType::IncludeParentTags, EGameplayTagMatchType::Explicit, EGameplayContainerMatchType::Any).Num() > 0;
+	return DoesTagContainerMatch(Other, EGameplayTagMatchType::IncludeParentTags, EGameplayTagMatchType::Explicit, EGameplayContainerMatchType::Any);
 }
 
 void FGameplayTagContainer::AppendTags(FGameplayTagContainer const& Other)
