GameplayTags - Add method for getting all children GameplayTag nodes

diff --git a/Engine/Source/Runtime/GameplayTags/Classes/GameplayTagsManager.h b/Engine/Source/Runtime/GameplayTags/Classes/GameplayTagsManager.h
--- a/Engine/Source/Runtime/GameplayTags/Classes/GameplayTagsManager.h
+++ b/Engine/Source/Runtime/GameplayTags/Classes/GameplayTagsManager.h
@@ -199,6 +199,8 @@ class GAMEPLAYTAGS_API UGameplayTagsManager : public UObject
 	 */
 	FGameplayTagContainer RequestGameplayTagParents(const FGameplayTag& GameplayTag) const;
 
+	FGameplayTagContainer RequestGameplayTagChildren(const FGameplayTag& GameplayTag) const;
+
 	/**
 	 * Checks if the tag is allowed to be created
 	 *
@@ -238,6 +240,8 @@ private:
 	 */
 	void AddParentTags(FGameplayTagContainer& TagContainer, const FGameplayTag& GameplayTag) const;
 
+	void AddChildrenTags(FGameplayTagContainer& TagContainer, const FGameplayTag& GameplayTag, bool RecurseAll=true) const;
+
 	/**
 	 * Helper function for GameplayTagsMatch to get all parents when doing a parent match,
 	 * NOTE: Must never be made public as it uses the FNames which should never be exposed
diff --git a/Engine/Source/Runtime/GameplayTags/Private/GameplayTagsManager.cpp b/Engine/Source/Runtime/GameplayTags/Private/GameplayTagsManager.cpp
--- a/Engine/Source/Runtime/GameplayTags/Private/GameplayTagsManager.cpp
+++ b/Engine/Source/Runtime/GameplayTags/Private/GameplayTagsManager.cpp
@@ -277,6 +277,14 @@ FGameplayTagContainer UGameplayTagsManager::RequestGameplayTagParents(const FGam
 	return TagContainer;
 }
 
+FGameplayTagContainer UGameplayTagsManager::RequestGameplayTagChildren(const FGameplayTag& GameplayTag) const
+{
+	FGameplayTagContainer TagContainer;
+	// Note this purposefully does not include the passed in GameplayTag in the container.
+	AddChildrenTags(TagContainer, GameplayTag, true);
+	return TagContainer;
+}
+
 bool UGameplayTagsManager::AddLeafTagToContainer(FGameplayTagContainer& TagContainer, FGameplayTag& Tag)
 {
 	// Check tag is not already in container
@@ -328,6 +336,31 @@ void UGameplayTagsManager::AddParentTags(FGameplayTagContainer& TagContainer, co
 	}
 }
 
+void UGameplayTagsManager::AddChildrenTags(FGameplayTagContainer& TagContainer, const FGameplayTag& GameplayTag, bool RecurseAll) const
+{
+	const TSharedPtr<FGameplayTagNode>* GameplayTagNode = GameplayTagNodeMap.Find(GameplayTag);
+	if (GameplayTagNode)
+	{
+		TArray< TSharedPtr<FGameplayTagNode> >& ChildrenNodes = (*GameplayTagNode)->GetChildTagNodes();
+		for (TSharedPtr<FGameplayTagNode> ChildNode : ChildrenNodes)
+		{
+			if (ChildNode.IsValid())
+			{
+				const FGameplayTag* Tag = GameplayTagNodeMap.FindKey(ChildNode);
+				if (Tag)
+				{
+					TagContainer.AddTag(*Tag);
+					if (RecurseAll)
+					{
+						AddChildrenTags(TagContainer, *Tag, true);
+					}
+				}
+			}
+
+		}
+	}
+}
+
 bool UGameplayTagsManager::GameplayTagsMatch(const FGameplayTag& GameplayTagOne, TEnumAsByte<EGameplayTagMatchType::Type> MatchTypeOne, const FGameplayTag& GameplayTagTwo, TEnumAsByte<EGameplayTagMatchType::Type> MatchTypeTwo) const
 {
 	TSet<FName> Tags1;
