Deprecation of old FName tags inside GameplayTagContainer; Causes engine obj. version bump
#ue4

diff --git a/Engine/Source/Runtime/Core/Public/UObject/ObjectVersion.h b/Engine/Source/Runtime/Core/Public/UObject/ObjectVersion.h
--- a/Engine/Source/Runtime/Core/Public/UObject/ObjectVersion.h
+++ b/Engine/Source/Runtime/Core/Public/UObject/ObjectVersion.h
@@ -552,7 +552,7 @@ enum EUnrealEngineObjectUE4Version
 	VER_UE4_FTEXT_HISTORY,
 	// Shift comments to the left to contain expressions properly
 	VER_UE4_FIX_MATERIAL_COMMENTS,
-	// Bone names stored as FName mean that we cant guarrentee the correct case on export, now we store a separate string for export purposes only
+	// Bone names stored as FName means that we can't guarantee the correct case on export, now we store a separate string for export purposes only
 	VER_UE4_STORE_BONE_EXPORT_NAMES,
 	// changed mesh emitter initial orientation to distribution
 	VER_UE4_MESH_EMITTER_INITIAL_ORIENTATION_DISTRIBUTION,
@@ -564,6 +564,8 @@ enum EUnrealEngineObjectUE4Version
 	VER_UE4_DEPRECATED_MOVEMENTCOMPONENT_MODIFIED_SPEEDS,
 	// rename CanBeCharacterBase
 	VER_UE4_RENAME_CANBECHARACTERBASE,
+	// Change GameplayTagContainers to have FGameplayTags instead of FNames; Required to fix-up native serialization
+	VER_UE4_GAMEPLAY_TAG_CONTAINER_TAG_TYPE_CHANGE,
 
 	// -----<new versions can be added before this line>-------------------------------------------------
 	// - this needs to be the last line (see note below)
diff --git a/Engine/Source/Runtime/GameplayTags/Classes/GameplayTagContainer.h b/Engine/Source/Runtime/GameplayTags/Classes/GameplayTagContainer.h
--- a/Engine/Source/Runtime/GameplayTags/Classes/GameplayTagContainer.h
+++ b/Engine/Source/Runtime/GameplayTags/Classes/GameplayTagContainer.h
@@ -230,10 +230,9 @@ protected:
 
 private:
 
-	// @todo: Remove this via engine version bump/native serialization soon
 	/** Array of gameplay tags */
 	UPROPERTY()
-	TArray<FName> Tags;
+	TArray<FName> Tags_DEPRECATED;
 
 	/** Array of gameplay tags */
 	UPROPERTY(VisibleAnywhere, Category=GameplayTags)
diff --git a/Engine/Source/Runtime/GameplayTags/Private/GameplayTagContainer.cpp b/Engine/Source/Runtime/GameplayTags/Private/GameplayTagContainer.cpp
--- a/Engine/Source/Runtime/GameplayTags/Private/GameplayTagContainer.cpp
+++ b/Engine/Source/Runtime/GameplayTags/Private/GameplayTagContainer.cpp
@@ -174,30 +174,34 @@ void FGameplayTagContainer::RemoveAllTags()
 
 bool FGameplayTagContainer::Serialize(FArchive& Ar)
 {
-	if (Ar.IsSaving())
+	const bool bOldTagVer = Ar.UE4Ver() < VER_UE4_GAMEPLAY_TAG_CONTAINER_TAG_TYPE_CHANGE;
+	
+	if (bOldTagVer)
 	{
-		FGameplayTagContainer ExpandedContainer = GetGameplayTagParents();
-		Tags.Empty(ExpandedContainer.Num());
-
-		for (auto TagIter = ExpandedContainer.CreateConstIterator(); TagIter; ++TagIter)
-		{
-			Tags.Add((*TagIter).GetTagName());
-		}
+		Ar << Tags_DEPRECATED;
 	}
-
-	Ar << Tags;
-
+	else
+	{
+		Ar << GameplayTags;
+	}
+	
 	if (Ar.IsLoading())
 	{
+		// Regardless of version, want loading to have a chance to handle redirects
 		RedirectTags();
-		
-		UGameplayTagsManager& TagManager = IGameplayTagsModule::Get().GetGameplayTagsManager();
-		for (auto It = Tags.CreateConstIterator(); It; ++It)
+
+		// If loading old version, add old tags to the new gameplay tags array so they can be saved out with the new version
+		if (bOldTagVer)
 		{
-			FGameplayTag Tag = TagManager.RequestGameplayTag(*It);
-			TagManager.AddLeafTagToContainer(*this, Tag);
+			UGameplayTagsManager& TagManager = IGameplayTagsModule::Get().GetGameplayTagsManager();
+			for (auto It = Tags_DEPRECATED.CreateConstIterator(); It; ++It)
+			{
+				FGameplayTag Tag = TagManager.RequestGameplayTag(*It);
+				TagManager.AddLeafTagToContainer(*this, Tag);
+			}
 		}
 	}
+
 	return true;
 }
 
@@ -217,10 +221,10 @@ void FGameplayTagContainer::RedirectTags()
 			FParse::Value( *It.Value(), TEXT("NewTagName="), NewTagName );
 
 			int32 TagIndex = 0;
-			if (Tags.Find(OldTagName, TagIndex))
+			if (Tags_DEPRECATED.Find(OldTagName, TagIndex))
 			{
-				Tags.RemoveAt(TagIndex);
-				Tags.AddUnique(NewTagName);
+				Tags_DEPRECATED.RemoveAt(TagIndex);
+				Tags_DEPRECATED.AddUnique(NewTagName);
 			}
 
 			FGameplayTag OldTag = TagManager.RequestGameplayTag(OldTagName);
