UE4: Fixed a crash caught by automation testing in gameplay tags that happens when there are no configured gameplay tag tables for the project.

diff --git a/Engine/Source/Runtime/GameplayTags/Private/GameplayTagsManager.cpp b/Engine/Source/Runtime/GameplayTags/Private/GameplayTagsManager.cpp
--- a/Engine/Source/Runtime/GameplayTags/Private/GameplayTagsManager.cpp
+++ b/Engine/Source/Runtime/GameplayTags/Private/GameplayTagsManager.cpp
@@ -109,7 +109,7 @@ struct FCompareFGameplayTagNodeByTag
 
 void UGameplayTagsManager::ConstructGameplayTagTree()
 {
-	if (!GameplayRootTag.IsValid() && GameplayTagTables.Num() > 0)
+	if (!GameplayRootTag.IsValid())
 	{
 		GameplayRootTag = MakeShareable(new FGameplayTagNode());
 		
@@ -122,6 +122,8 @@ void UGameplayTagsManager::ConstructGameplayTagTree()
 
 void UGameplayTagsManager::PopulateTreeFromDataTable(class UDataTable* InTable)
 {
+	checkf(GameplayRootTag.IsValid(), TEXT("ConstructGameplayTagTree() must be called before PopulateTreeFromDataTable()"));
+
 	TArray< TSharedPtr<FGameplayTagNode> >& GameplayRootTags = GameplayRootTag->GetChildTagNodes();
 	static const FString ContextString(TEXT("UNKNOWN"));
 	const int32 NumRows = InTable->RowMap.Num();
