GameplayTagsContainer - add tag fast function

diff --git a/Engine/Source/Runtime/GameplayTags/Classes/GameplayTagContainer.h b/Engine/Source/Runtime/GameplayTags/Classes/GameplayTagContainer.h
--- a/Engine/Source/Runtime/GameplayTags/Classes/GameplayTagContainer.h
+++ b/Engine/Source/Runtime/GameplayTags/Classes/GameplayTagContainer.h
@@ -175,6 +175,15 @@ struct GAMEPLAYTAGS_API FGameplayTagContainer
 	 */
 	void AddTag(const FGameplayTag& TagToAdd);
 
+	/**
+	 * Add the specified tag to the container without checking for uniqueness
+	 *
+	 * @param TagToAdd Tag to add to the container
+	 * 
+	 * Useful when building container from another data struct (TMap for example)
+	 */
+	void AddTagFast(const FGameplayTag& TagToAdd);
+
 	/**
 	 * Tag to remove from the container
 	 * 
diff --git a/Engine/Source/Runtime/GameplayTags/Private/GameplayTagContainer.cpp b/Engine/Source/Runtime/GameplayTags/Private/GameplayTagContainer.cpp
--- a/Engine/Source/Runtime/GameplayTags/Private/GameplayTagContainer.cpp
+++ b/Engine/Source/Runtime/GameplayTags/Private/GameplayTagContainer.cpp
@@ -162,6 +162,11 @@ void FGameplayTagContainer::AddTag(const FGameplayTag& TagToAdd)
 	}
 }
 
+void FGameplayTagContainer::AddTagFast(const FGameplayTag& TagToAdd)
+{
+	GameplayTags.Add(TagToAdd);
+}
+
 void FGameplayTagContainer::RemoveTag(FGameplayTag TagToRemove)
 {
 	GameplayTags.Remove(TagToRemove);
