Support for adding GameplayTags through ProjectSettings (in addition to importing dictionary uassets)

diff --git a/Engine/Source/Editor/GameplayTagsEditor/GameplayTagsEditor.Build.cs b/Engine/Source/Editor/GameplayTagsEditor/GameplayTagsEditor.Build.cs
--- a/Engine/Source/Editor/GameplayTagsEditor/GameplayTagsEditor.Build.cs
+++ b/Engine/Source/Editor/GameplayTagsEditor/GameplayTagsEditor.Build.cs
@@ -33,6 +33,12 @@ namespace UnrealBuildTool.Rules
 					"UnrealEd",
 				}
 			);
+
+            PrivateIncludePathModuleNames.AddRange(
+                new string[] {
+				    "Settings"
+			    }
+            );
 		}
 	}
 }
diff --git a/Engine/Source/Editor/GameplayTagsEditor/Private/GameplayTagsEditorModule.cpp b/Engine/Source/Editor/GameplayTagsEditor/Private/GameplayTagsEditorModule.cpp
--- a/Engine/Source/Editor/GameplayTagsEditor/Private/GameplayTagsEditorModule.cpp
+++ b/Engine/Source/Editor/GameplayTagsEditor/Private/GameplayTagsEditorModule.cpp
@@ -4,6 +4,10 @@
 #include "GameplayTagsGraphPanelPinFactory.h"
 #include "GameplayTagContainerCustomization.h"
 #include "GameplayTagCustomization.h"
+#include "GameplayTagsSettings.h"
+#include "Settings.h"
+
+#define LOCTEXT_NAMESPACE "GameplayTagEditor"
 
 class FGameplayTagsEditorModule : public IGameplayTagsEditorModule
 {
@@ -24,11 +28,27 @@ void FGameplayTagsEditorModule::StartupModule()
 
 	TSharedPtr<FGameplayTagsGraphPanelPinFactory> GameplayTagsGraphPanelPinFactory = MakeShareable( new FGameplayTagsGraphPanelPinFactory() );
 	FEdGraphUtilities::RegisterVisualPinFactory(GameplayTagsGraphPanelPinFactory);
-	
+
+	if (ISettingsModule* SettingsModule = ISettingsModule::Get())
+	{
+		SettingsModule->RegisterSettings("Project", "Project", "GameplayTags",
+			LOCTEXT("GameplayTagSettingsName", "GameplayTags"),
+			LOCTEXT("GameplayTagSettingsNameDesc", "GameplayTag Settings"),
+			GetMutableDefault<UGameplayTagsSettings>()
+			);
+	}
 }
 
 void FGameplayTagsEditorModule::ShutdownModule()
 {
 	// This function may be called during shutdown to clean up your module.  For modules that support dynamic reloading,
 	// we call this function before unloading the module.
+	
+	
+	if (ISettingsModule* SettingsModule = ISettingsModule::Get())
+	{
+		SettingsModule->UnregisterSettings("Project", "Project", "GameplayTags");
+	}
 }
+
+#undef LOCTEXT_NAMESPACE
\ No newline at end of file
diff --git a/Engine/Source/Runtime/GameplayTags/Classes/GameplayTagsManager.h b/Engine/Source/Runtime/GameplayTags/Classes/GameplayTagsManager.h
--- a/Engine/Source/Runtime/GameplayTags/Classes/GameplayTagsManager.h
+++ b/Engine/Source/Runtime/GameplayTags/Classes/GameplayTagsManager.h
@@ -238,6 +238,8 @@ private:
 	/** Helper function to populate the tag tree from each table */
 	void PopulateTreeFromDataTable(class UDataTable* Table);
 
+	void AddTagTableRow(const FGameplayTagTableRow& TagRow);
+
 	/**
 	 * Helper function for RequestGameplayTagParents to add all parents to the container
 	 *
diff --git a/Engine/Source/Runtime/GameplayTags/Classes/GameplayTagsSettings.h b/Engine/Source/Runtime/GameplayTags/Classes/GameplayTagsSettings.h
new file mode 100644
--- /dev/null
+++ b/Engine/Source/Runtime/GameplayTags/Classes/GameplayTagsSettings.h
@@ -0,0 +1,34 @@
+// Copyright 1998-2014 Epic Games, Inc. All Rights Reserved.
+
+#pragma once
+
+#include "GameplayTagContainer.h"
+#include "GameplayTagsSettings.generated.h"
+
+/**
+ *	Class for importing GameplayTags directly from a config file.
+ *	FGameplayTagsEditorModule::StartupModule adds this class to the Project Settings menu to be edited.
+ *	Editing this in Project Settings will output changes to Config/DEfaultEngine.ini.
+ *	
+ *	Primary advantages of this approach are:
+ *	-Adding new tags doesn't require checking out external and editing file (CSV or xls) then reimporting.
+ *	-New tags are mergeable since .ini are text and non exclusive checkout.
+ *	
+ *	To do:
+ *	-Better support could be added for adding new tags. We could match existing tags and autocomplete subtags as
+ *	the user types (e.g, autocomplete 'Damage.Physical' as the user is adding a 'Damage.Physical.Slash' tag).
+ *	
+ */
+UCLASS(config=Engine, defaultconfig, notplaceable)
+class GAMEPLAYTAGS_API UGameplayTagsSettings : public UObject
+{
+	GENERATED_UCLASS_BODY()
+
+	UPROPERTY(config, EditAnywhere, Category=GameplayTags)
+	TArray<FString>		GameplayTags;
+
+#if WITH_EDITOR
+	virtual void PostEditChangeProperty(FPropertyChangedEvent& PropertyChangedEvent) override;
+#endif
+
+};
diff --git a/Engine/Source/Runtime/GameplayTags/Private/GameplayTagsManager.cpp b/Engine/Source/Runtime/GameplayTags/Private/GameplayTagsManager.cpp
--- a/Engine/Source/Runtime/GameplayTags/Private/GameplayTagsManager.cpp
+++ b/Engine/Source/Runtime/GameplayTags/Private/GameplayTagsManager.cpp
@@ -1,8 +1,12 @@
 // Copyright 1998-2014 Epic Games, Inc. All Rights Reserved.
 
 #include "GameplayTagsModulePrivatePCH.h"
+#include "GameplayTagsSettings.h"
 //#include "AssetRegistryModule.h"
 
+#pragma optimize( "", off )
+
+
 #if WITH_EDITOR
 #include "UnrealEd.h"
 #endif
@@ -38,6 +42,7 @@ void UGameplayTagsManager::LoadGameplayTagTables(TArray<FString>& TagTableNames)
 			GameplayTagTables.Add(TagTable);
 		}
 	}
+
 #if WITH_EDITOR
 	// Hook into notifications for object re-imports so that the gameplay tag tree can be reconstructed if the table changes
 	if (GIsEditor && GameplayTagTables.Num() > 0 && !RegisteredObjectReimport)
@@ -117,47 +122,70 @@ void UGameplayTagsManager::ConstructGameplayTagTree()
 		{
 			PopulateTreeFromDataTable(*It);
 		}
+
+		// Load any GameplayTagSettings from config (their default object)
+		for (TObjectIterator<UClass> ClassIt; ClassIt; ++ClassIt)
+		{
+			UClass* Class = *ClassIt;
+			if (!Class->IsChildOf<UGameplayTagsSettings>() || Class->HasAnyClassFlags(CLASS_Abstract))
+			{
+				continue;
+			}
+
+			for (FString TagStr : Class->GetDefaultObject<UGameplayTagsSettings>()->GameplayTags)
+			{
+				FGameplayTagTableRow TableRow;
+				TableRow.Tag = TagStr;
+				AddTagTableRow(TableRow);
+			}
+		}
 	}
 }
 
 void UGameplayTagsManager::PopulateTreeFromDataTable(class UDataTable* InTable)
 {
 	checkf(GameplayRootTag.IsValid(), TEXT("ConstructGameplayTagTree() must be called before PopulateTreeFromDataTable()"));
-
-	TArray< TSharedPtr<FGameplayTagNode> >& GameplayRootTags = GameplayRootTag->GetChildTagNodes();
 	static const FString ContextString(TEXT("UNKNOWN"));
+	
 	const int32 NumRows = InTable->RowMap.Num();
 	for (int32 RowIdx = 0; RowIdx < NumRows; ++RowIdx)
 	{
 		FGameplayTagTableRow* TagRow = InTable->FindRow<FGameplayTagTableRow>(*FString::Printf(TEXT("%d"), RowIdx), ContextString);
 		if (TagRow)
 		{
-			// Split the tag text on the "." delimiter to establish tag depth and then insert each tag into the
-			// gameplay tag tree
-			TArray<FString> SubTags;
-			TagRow->Tag.ParseIntoArray(&SubTags, TEXT("."), true);
+			AddTagTableRow(*TagRow);
+		}
+	}
+	GameplayRootTag->GetChildTagNodes().Sort(FCompareFGameplayTagNodeByTag());
+}
 
-			if (SubTags.Num() > 0)
-			{
-				int32 InsertionIdx = InsertTagIntoNodeArray(*SubTags[0], NULL, GameplayRootTags, SubTags.Num() == 1 ? TagRow->CategoryText : FText());
-				TSharedPtr<FGameplayTagNode> CurNode = GameplayRootTags[InsertionIdx];
+void UGameplayTagsManager::AddTagTableRow(const FGameplayTagTableRow& TagRow)
+{
+	TArray< TSharedPtr<FGameplayTagNode> >& GameplayRootTags = GameplayRootTag->GetChildTagNodes();
 
-				for (int32 SubTagIdx = 1; SubTagIdx < SubTags.Num(); ++SubTagIdx)
-				{
-					TArray< TSharedPtr<FGameplayTagNode> >& ChildTags = CurNode.Get()->GetChildTagNodes();
-					FText Description;
-					if (SubTagIdx == SubTags.Num() - 1)
-					{
-						Description = TagRow->CategoryText;
-					}
-					InsertionIdx = InsertTagIntoNodeArray(*SubTags[SubTagIdx], CurNode, ChildTags, Description);
+	// Split the tag text on the "." delimiter to establish tag depth and then insert each tag into the
+	// gameplay tag tree
+	TArray<FString> SubTags;
+	TagRow.Tag.ParseIntoArray(&SubTags, TEXT("."), true);
 
-					CurNode = ChildTags[InsertionIdx];
-				}
+	if (SubTags.Num() > 0)
+	{
+		int32 InsertionIdx = InsertTagIntoNodeArray(*SubTags[0], NULL, GameplayRootTags, SubTags.Num() == 1 ? TagRow.CategoryText : FText());
+		TSharedPtr<FGameplayTagNode> CurNode = GameplayRootTags[InsertionIdx];
+
+		for (int32 SubTagIdx = 1; SubTagIdx < SubTags.Num(); ++SubTagIdx)
+		{
+			TArray< TSharedPtr<FGameplayTagNode> >& ChildTags = CurNode.Get()->GetChildTagNodes();
+			FText Description;
+			if (SubTagIdx == SubTags.Num() - 1)
+			{
+				Description = TagRow.CategoryText;
 			}
+			InsertionIdx = InsertTagIntoNodeArray(*SubTags[SubTagIdx], CurNode, ChildTags, Description);
+
+			CurNode = ChildTags[InsertionIdx];
 		}
 	}
-	GameplayRootTags.Sort(FCompareFGameplayTagNodeByTag());
 }
 
 UGameplayTagsManager::~UGameplayTagsManager()
diff --git a/Engine/Source/Runtime/GameplayTags/Private/GameplayTagsSettings.cpp b/Engine/Source/Runtime/GameplayTags/Private/GameplayTagsSettings.cpp
new file mode 100644
--- /dev/null
+++ b/Engine/Source/Runtime/GameplayTags/Private/GameplayTagsSettings.cpp
@@ -0,0 +1,26 @@
+// Copyright 1998-2014 Epic Games, Inc. All Rights Reserved.
+
+#include "GameplayTagsModulePrivatePCH.h"
+#include "GameplayTagsModule.h"
+
+UGameplayTagsSettings::UGameplayTagsSettings(const class FPostConstructInitializeProperties& PCIP)
+: Super(PCIP)
+{
+
+
+}
+
+#if WITH_EDITOR
+void UGameplayTagsSettings::PostEditChangeProperty(FPropertyChangedEvent& PropertyChangedEvent)
+{
+	Super::PostEditChangeProperty(PropertyChangedEvent);
+	if (PropertyChangedEvent.Property)
+	{
+		if (PropertyChangedEvent.Property->GetFName() == GET_MEMBER_NAME_CHECKED(UGameplayTagsSettings, GameplayTags))
+		{
+			IGameplayTagsModule::Get().GetGameplayTagsManager().DestroyGameplayTagTree();
+			IGameplayTagsModule::Get().GetGameplayTagsManager().ConstructGameplayTagTree();
+		}
+	}
+}
+#endif
