Make ImportFromConfig a config option. Disabled by default, enabled in Orion.

diff --git a/Engine/Source/Editor/GameplayTagsEditor/Private/GameplayTagsEditorModule.cpp b/Engine/Source/Editor/GameplayTagsEditor/Private/GameplayTagsEditorModule.cpp
--- a/Engine/Source/Editor/GameplayTagsEditor/Private/GameplayTagsEditorModule.cpp
+++ b/Engine/Source/Editor/GameplayTagsEditor/Private/GameplayTagsEditorModule.cpp
@@ -29,13 +29,16 @@ void FGameplayTagsEditorModule::StartupModule()
 	TSharedPtr<FGameplayTagsGraphPanelPinFactory> GameplayTagsGraphPanelPinFactory = MakeShareable( new FGameplayTagsGraphPanelPinFactory() );
 	FEdGraphUtilities::RegisterVisualPinFactory(GameplayTagsGraphPanelPinFactory);
 
-	if (ISettingsModule* SettingsModule = ISettingsModule::Get())
+	if (UGameplayTagsManager::ShouldImportTagsFromINI())
 	{
-		SettingsModule->RegisterSettings("Project", "Project", "GameplayTags",
-			LOCTEXT("GameplayTagSettingsName", "GameplayTags"),
-			LOCTEXT("GameplayTagSettingsNameDesc", "GameplayTag Settings"),
-			GetMutableDefault<UGameplayTagsSettings>()
-			);
+		if (ISettingsModule* SettingsModule = ISettingsModule::Get())
+		{
+			SettingsModule->RegisterSettings("Project", "Project", "GameplayTags",
+				LOCTEXT("GameplayTagSettingsName", "GameplayTags"),
+				LOCTEXT("GameplayTagSettingsNameDesc", "GameplayTag Settings"),
+				GetMutableDefault<UGameplayTagsSettings>()
+				);
+		}
 	}
 }
 
diff --git a/Engine/Source/Runtime/GameplayTags/Classes/GameplayTagsManager.h b/Engine/Source/Runtime/GameplayTags/Classes/GameplayTagsManager.h
--- a/Engine/Source/Runtime/GameplayTags/Classes/GameplayTagsManager.h
+++ b/Engine/Source/Runtime/GameplayTags/Classes/GameplayTagsManager.h
@@ -230,6 +230,9 @@ class GAMEPLAYTAGS_API UGameplayTagsManager : public UObject
 	 */
 	bool GameplayTagsMatch(const FGameplayTag& GameplayTagOne, TEnumAsByte<EGameplayTagMatchType::Type> MatchTypeOne, const FGameplayTag& GameplayTagTwo, TEnumAsByte<EGameplayTagMatchType::Type> MatchTypeTwo) const;
 
+	/** Returns true if we should import tags from UGameplayTagsSettings objects (configured by INI files) */
+	static bool ShouldImportTagsFromINI();
+
 private:
 
 	friend class FGameplayTagTest;
diff --git a/Engine/Source/Runtime/GameplayTags/Private/GameplayTagsManager.cpp b/Engine/Source/Runtime/GameplayTags/Private/GameplayTagsManager.cpp
--- a/Engine/Source/Runtime/GameplayTags/Private/GameplayTagsManager.cpp
+++ b/Engine/Source/Runtime/GameplayTags/Private/GameplayTagsManager.cpp
@@ -123,25 +123,38 @@ void UGameplayTagsManager::ConstructGameplayTagTree()
 			PopulateTreeFromDataTable(*It);
 		}
 
-		// Load any GameplayTagSettings from config (their default object)
-		for (TObjectIterator<UClass> ClassIt; ClassIt; ++ClassIt)
+		if (ShouldImportTagsFromINI())
 		{
-			UClass* Class = *ClassIt;
-			if (!Class->IsChildOf<UGameplayTagsSettings>() || Class->HasAnyClassFlags(CLASS_Abstract))
+			// Load any GameplayTagSettings from config (their default object)
+			for (TObjectIterator<UClass> ClassIt; ClassIt; ++ClassIt)
 			{
-				continue;
-			}
+				UClass* Class = *ClassIt;
+				if (!Class->IsChildOf<UGameplayTagsSettings>() || Class->HasAnyClassFlags(CLASS_Abstract))
+				{
+					continue;
+				}
 
-			for (FString TagStr : Class->GetDefaultObject<UGameplayTagsSettings>()->GameplayTags)
-			{
-				FGameplayTagTableRow TableRow;
-				TableRow.Tag = TagStr;
-				AddTagTableRow(TableRow);
+				for (FString TagStr : Class->GetDefaultObject<UGameplayTagsSettings>()->GameplayTags)
+				{
+					FGameplayTagTableRow TableRow;
+					TableRow.Tag = TagStr;
+					AddTagTableRow(TableRow);
+				}
+
+		
 			}
+			GameplayRootTag->GetChildTagNodes().Sort(FCompareFGameplayTagNodeByTag());
 		}
 	}
 }
 
+bool UGameplayTagsManager::ShouldImportTagsFromINI()
+{
+	bool ImportFromINI = false;
+	GConfig->GetBool(TEXT("GameplayTags"), TEXT("ImportTagsFromConfig"), ImportFromINI, GEngineIni);
+	return ImportFromINI;
+}
+
 void UGameplayTagsManager::PopulateTreeFromDataTable(class UDataTable* InTable)
 {
 	checkf(GameplayRootTag.IsValid(), TEXT("ConstructGameplayTagTree() must be called before PopulateTreeFromDataTable()"));
