remove pragma optimize off

diff --git a/Engine/Source/Runtime/GameplayTags/Private/GameplayTagsManager.cpp b/Engine/Source/Runtime/GameplayTags/Private/GameplayTagsManager.cpp
--- a/Engine/Source/Runtime/GameplayTags/Private/GameplayTagsManager.cpp
+++ b/Engine/Source/Runtime/GameplayTags/Private/GameplayTagsManager.cpp
@@ -4,8 +4,6 @@
 #include "GameplayTagsSettings.h"
 //#include "AssetRegistryModule.h"
 
-#pragma optimize( "", off )
-
 
 #if WITH_EDITOR
 #include "UnrealEd.h"
