Compile fix. Let gameplayTagContainers be created with cstor that takes single gameplaytag.

diff --git a/Engine/Source/Runtime/GameplayTags/Classes/GameplayTagContainer.h b/Engine/Source/Runtime/GameplayTags/Classes/GameplayTagContainer.h
--- a/Engine/Source/Runtime/GameplayTags/Classes/GameplayTagContainer.h
+++ b/Engine/Source/Runtime/GameplayTags/Classes/GameplayTagContainer.h
@@ -107,6 +107,7 @@ struct GAMEPLAYTAGS_API FGameplayTagContainer
 	/** Constructors */
 	FGameplayTagContainer();
 	FGameplayTagContainer(FGameplayTagContainer const& Other);
+	FGameplayTagContainer(const FGameplayTag& Tag);
 	virtual ~FGameplayTagContainer() {}
 
 	/** Assignment/Equality operators */
diff --git a/Engine/Source/Runtime/GameplayTags/Private/GameplayTagContainer.cpp b/Engine/Source/Runtime/GameplayTags/Private/GameplayTagContainer.cpp
--- a/Engine/Source/Runtime/GameplayTags/Private/GameplayTagContainer.cpp
+++ b/Engine/Source/Runtime/GameplayTags/Private/GameplayTagContainer.cpp
@@ -10,6 +10,11 @@ FGameplayTagContainer::FGameplayTagContainer(FGameplayTagContainer const& Other)
 	*this = Other;
 }
 
+FGameplayTagContainer::FGameplayTagContainer(const FGameplayTag& Tag)
+{
+	AddTag(Tag);
+}
+
 FGameplayTagContainer& FGameplayTagContainer::operator=(FGameplayTagContainer const& Other)
 {
 	// Guard against self-assignment
