GameplayTags: range based loop support, method for getting direct parent of gameplaytag

diff --git a/Engine/Source/Runtime/GameplayTags/Classes/GameplayTagContainer.h b/Engine/Source/Runtime/GameplayTags/Classes/GameplayTagContainer.h
--- a/Engine/Source/Runtime/GameplayTags/Classes/GameplayTagContainer.h
+++ b/Engine/Source/Runtime/GameplayTags/Classes/GameplayTagContainer.h
@@ -86,6 +86,9 @@ struct GAMEPLAYTAGS_API FGameplayTag
 		return Ar;
 	}
 
+	/** Returns direct parent GameplayTag of this GameplayTag */
+	FGameplayTag RequestDirectParent() const;
+
 private:
 
 	/** Intentionally private so only the tag manager can use */
@@ -224,6 +227,17 @@ struct GAMEPLAYTAGS_API FGameplayTagContainer
 		return GameplayTags.CreateConstIterator();
 	}
 
+	FGameplayTag First() const
+	{
+		return GameplayTags.Num() > 0 ? GameplayTags[0] : FGameplayTag();
+	}
+
+	FGameplayTag Last() const
+	{
+		return GameplayTags.Num() > 0 ? GameplayTags.Last() : FGameplayTag();
+	}
+
+
 protected:
 	/**
 	* Returns true if the tags in this container match the tags in OtherContainer for the specified matching types.
@@ -247,6 +261,14 @@ private:
 	/** Array of gameplay tags */
 	UPROPERTY(VisibleAnywhere, Category=GameplayTags)
 	TArray<FGameplayTag> GameplayTags;
+
+	/**
+	 * DO NOT USE DIRECTLY
+	 * STL-like iterators to enable range-based for loop support.
+	 */
+	
+	FORCEINLINE friend TArray<FGameplayTag>::TConstIterator begin(const FGameplayTagContainer& Array) { return Array.CreateConstIterator(); }
+	FORCEINLINE friend TArray<FGameplayTag>::TConstIterator end(const FGameplayTagContainer& Array) { return TArray<FGameplayTag>::TConstIterator(Array.GameplayTags, Array.GameplayTags.Num()); }
 };
 
 template<>
diff --git a/Engine/Source/Runtime/GameplayTags/Classes/GameplayTagsManager.h b/Engine/Source/Runtime/GameplayTags/Classes/GameplayTagsManager.h
--- a/Engine/Source/Runtime/GameplayTags/Classes/GameplayTagsManager.h
+++ b/Engine/Source/Runtime/GameplayTags/Classes/GameplayTagsManager.h
@@ -175,10 +175,12 @@ class GAMEPLAYTAGS_API UGameplayTagsManager : public UObject
 	 *
 	 * @param TagName The Name of the tag to search for
 	 * 
+	 * @param ErrorIfNotfound: ensure() that tag exists.
+	 * 
 	 * @return Will return the corresponding FGameplayTag or an empty one if not found.
 	 */
 	UFUNCTION(BlueprintCallable, Category="GameplayTags")
-	FGameplayTag RequestGameplayTag(FName TagName) const;
+	FGameplayTag RequestGameplayTag(FName TagName, bool ErrorIfNotFound=true) const;
 
 	/**
 	 * Adds a tag to the container and removes any direct parents, wont add if child already exists
@@ -201,6 +203,8 @@ class GAMEPLAYTAGS_API UGameplayTagsManager : public UObject
 
 	FGameplayTagContainer RequestGameplayTagChildren(const FGameplayTag& GameplayTag) const;
 
+	FGameplayTag RequestGameplayTagDirectParent(const FGameplayTag& GameplayTag) const;
+
 	/**
 	 * Checks if the tag is allowed to be created
 	 *
diff --git a/Engine/Source/Runtime/GameplayTags/Private/GameplayTagContainer.cpp b/Engine/Source/Runtime/GameplayTags/Private/GameplayTagContainer.cpp
--- a/Engine/Source/Runtime/GameplayTags/Private/GameplayTagContainer.cpp
+++ b/Engine/Source/Runtime/GameplayTags/Private/GameplayTagContainer.cpp
@@ -291,3 +291,8 @@ FGameplayTag::FGameplayTag(FName Name)
 {
 	check(IGameplayTagsModule::Get().GetGameplayTagsManager().ValidateTagCreation(Name));
 }
+
+FGameplayTag FGameplayTag::RequestDirectParent() const
+{
+	return IGameplayTagsModule::Get().GetGameplayTagsManager().RequestGameplayTagDirectParent(*this);
+}
\ No newline at end of file
diff --git a/Engine/Source/Runtime/GameplayTags/Private/GameplayTagsManager.cpp b/Engine/Source/Runtime/GameplayTags/Private/GameplayTagsManager.cpp
--- a/Engine/Source/Runtime/GameplayTags/Private/GameplayTagsManager.cpp
+++ b/Engine/Source/Runtime/GameplayTags/Private/GameplayTagsManager.cpp
@@ -303,12 +303,15 @@ void UGameplayTagsManager::OnObjectReimported(UFactory* ImportFactory, UObject*
 }
 #endif // WITH_EDITOR
 
-FGameplayTag UGameplayTagsManager::RequestGameplayTag(FName TagName) const
+FGameplayTag UGameplayTagsManager::RequestGameplayTag(FName TagName, bool ErrorIfNotFound) const
 {
 	const FGameplayTag* Tag = GameplayTagMap.Find(TagName);
 	if (!Tag)
 	{ 
-		ensureMsgf(false, TEXT("Requested Tag %s was not found. Check tag data table."), *TagName.ToString());
+		if (ErrorIfNotFound)
+		{
+			ensureMsgf(false, TEXT("Requested Tag %s was not found. Check tag data table."), *TagName.ToString());
+		}
 		return FGameplayTag();
 	}
 	return *Tag;
@@ -330,6 +333,24 @@ FGameplayTagContainer UGameplayTagsManager::RequestGameplayTagChildren(const FGa
 	return TagContainer;
 }
 
+FGameplayTag UGameplayTagsManager::RequestGameplayTagDirectParent(const FGameplayTag& GameplayTag) const
+{
+	const TSharedPtr<FGameplayTagNode>* GameplayTagNode = GameplayTagNodeMap.Find(GameplayTag);
+	if (GameplayTagNode)
+	{
+		TSharedPtr<FGameplayTagNode> Parent = (*GameplayTagNode)->GetParentTagNode().Pin();
+		if (Parent.IsValid())
+		{
+			const FGameplayTag* Tag = GameplayTagNodeMap.FindKey(Parent);
+			if (Tag)
+			{
+				return *Tag;
+			}
+		}
+	}
+	return FGameplayTag();
+}
+
 bool UGameplayTagsManager::AddLeafTagToContainer(FGameplayTagContainer& TagContainer, FGameplayTag& Tag)
 {
 	// Check tag is not already in container
