#UE4

- Adding an empty gameplay tag container

diff --git a/Engine/Source/Runtime/GameplayTags/Classes/GameplayTagContainer.h b/Engine/Source/Runtime/GameplayTags/Classes/GameplayTagContainer.h
--- a/Engine/Source/Runtime/GameplayTags/Classes/GameplayTagContainer.h
+++ b/Engine/Source/Runtime/GameplayTags/Classes/GameplayTagContainer.h
@@ -240,6 +240,8 @@ struct GAMEPLAYTAGS_API FGameplayTagContainer
 		return GameplayTags.Num() > 0 ? GameplayTags.Last() : FGameplayTag();
 	}
 
+	/** An empty Gameplay Tag Container */
+	static const FGameplayTagContainer EmptyContainer;
 
 protected:
 	/**
diff --git a/Engine/Source/Runtime/GameplayTags/Private/GameplayTagContainer.cpp b/Engine/Source/Runtime/GameplayTags/Private/GameplayTagContainer.cpp
--- a/Engine/Source/Runtime/GameplayTags/Private/GameplayTagContainer.cpp
+++ b/Engine/Source/Runtime/GameplayTags/Private/GameplayTagContainer.cpp
@@ -2,6 +2,8 @@
 
 #include "GameplayTagsModulePrivatePCH.h"
 
+const FGameplayTagContainer FGameplayTagContainer::EmptyContainer;
+
 FGameplayTagContainer::FGameplayTagContainer()
 {}
 
