expose ErrorIfNotFound parameter to IGameplayTagsModule::RequestGameplayTag.

diff --git a/Engine/Source/Runtime/GameplayTags/Public/GameplayTagsModule.h b/Engine/Source/Runtime/GameplayTags/Public/GameplayTagsModule.h
--- a/Engine/Source/Runtime/GameplayTags/Public/GameplayTagsModule.h
+++ b/Engine/Source/Runtime/GameplayTags/Public/GameplayTagsModule.h
@@ -46,10 +46,10 @@ public:
 	 * 
 	 * @return Gameplay tag associated with the specified name. Will be marked invalid if tag not found
 	 */
-	static FGameplayTag RequestGameplayTag(FName InTagName)
+	static FGameplayTag RequestGameplayTag(FName InTagName, bool ErrorIfNotFound=true)
 	{
 		IGameplayTagsModule& GameplayTagsModule = IGameplayTagsModule::Get();
-		return GameplayTagsModule.GetGameplayTagsManager().RequestGameplayTag(InTagName);
+		return GameplayTagsModule.GetGameplayTagsManager().RequestGameplayTag(InTagName, ErrorIfNotFound);
 	}
 };
 
