FGameplayTagContainer::AppendMatchingTags()

* Function for adding tags that match between two containers to another container.

--------
Integrated using branch UE4-To-UE4-Orion (reversed) of change#2577403 by [REDACTED_PERSON] on 2015/06/04 15:17:09.

diff --git a/Engine/Source/Runtime/GameplayTags/Classes/GameplayTagContainer.h b/Engine/Source/Runtime/GameplayTags/Classes/GameplayTagContainer.h
--- a/Engine/Source/Runtime/GameplayTags/Classes/GameplayTagContainer.h
+++ b/Engine/Source/Runtime/GameplayTags/Classes/GameplayTagContainer.h
@@ -186,6 +186,14 @@ struct GAMEPLAYTAGS_API FGameplayTagContainer
 	 */
 	virtual void AppendTags(FGameplayTagContainer const& Other);
 
+	/** 
+	 * Adds all the tags that match between the two specified containers to this container 
+	 *
+	 * @param OtherA TagContainer that has the matching tags you want to add to this container 
+	 * @param OtherB TagContainer used to check for matching tags
+	 */
+	virtual void AppendMatchingTags(FGameplayTagContainer const& OtherA, FGameplayTagContainer const& OtherB);
+
 	/**
 	 * Add the specified tag to the container
 	 *
diff --git a/Engine/Source/Runtime/GameplayTags/Private/GameplayTagContainer.cpp b/Engine/Source/Runtime/GameplayTags/Private/GameplayTagContainer.cpp
--- a/Engine/Source/Runtime/GameplayTags/Private/GameplayTagContainer.cpp
+++ b/Engine/Source/Runtime/GameplayTags/Private/GameplayTagContainer.cpp
@@ -186,6 +186,17 @@ void FGameplayTagContainer::AppendTags(FGameplayTagContainer const& Other)
 	}
 }
 
+void FGameplayTagContainer::AppendMatchingTags(FGameplayTagContainer const& OtherA, FGameplayTagContainer const& OtherB)
+{
+	for(TArray<FGameplayTag>::TConstIterator It(OtherA.GameplayTags); It; ++It)
+	{
+		if (OtherB.HasTag(*It, EGameplayTagMatchType::IncludeParentTags, EGameplayTagMatchType::Explicit))
+		{
+			AddTag(*It);		
+		}
+	}
+}
+
 void FGameplayTagContainer::AddTag(const FGameplayTag& TagToAdd)
 {
 	if (TagToAdd.IsValid())
