First pass at gameplay tag queries.

diff --git a/Engine/Source/Editor/GameplayTagsEditor/Private/GameplayTagQueryCustomization.cpp b/Engine/Source/Editor/GameplayTagsEditor/Private/GameplayTagQueryCustomization.cpp
new file mode 100644
--- /dev/null
+++ b/Engine/Source/Editor/GameplayTagsEditor/Private/GameplayTagQueryCustomization.cpp
@@ -0,0 +1,137 @@
+// Copyright 1998-2015 Epic Games, Inc. All Rights Reserved.
+
+#include "GameplayTagsEditorModulePrivatePCH.h"
+#include "GameplayTagQueryCustomization.h"
+#include "Editor/PropertyEditor/Public/PropertyEditing.h"
+#include "MainFrame.h"
+#include "SGameplayTagQueryWidget.h"
+#include "GameplayTagContainer.h"
+#include "ScopedTransaction.h"
+
+#define LOCTEXT_NAMESPACE "GameplayTagQueryCustomization"
+
+void FGameplayTagQueryCustomization::CustomizeHeader(TSharedRef<class IPropertyHandle> InStructPropertyHandle, class FDetailWidgetRow& HeaderRow, IPropertyTypeCustomizationUtils& StructCustomizationUtils)
+{
+	StructPropertyHandle = InStructPropertyHandle;
+
+	BuildEditableQueryList();
+
+	HeaderRow
+	.NameContent()
+	[
+		StructPropertyHandle->CreatePropertyNameWidget()
+	]
+	.ValueContent()
+	.MaxDesiredWidth(512)
+	[
+		SNew(SHorizontalBox)
+		+SHorizontalBox::Slot()
+		.AutoWidth()
+		.VAlign(VAlign_Center)
+		[
+			SNew(SVerticalBox)
+			+SVerticalBox::Slot()
+			.AutoHeight()
+			[
+				SNew(SButton)
+				.Text(this, &FGameplayTagQueryCustomization::GetEditButtonText)
+				.OnClicked(this, &FGameplayTagQueryCustomization::OnEditButtonClicked)
+			]
+		]
+	];
+
+	GEditor->RegisterForUndo(this);
+}
+
+FText FGameplayTagQueryCustomization::GetEditButtonText() const
+{
+	bool const bReadOnly = StructPropertyHandle->GetProperty()->HasAnyPropertyFlags(CPF_EditConst);
+	return
+		bReadOnly
+		? LOCTEXT("GameplayTagQueryCustomization_View", "View...")
+		: LOCTEXT("GameplayTagQueryCustomization_Edit", "Edit...");
+}
+
+FReply FGameplayTagQueryCustomization::OnEditButtonClicked()
+{
+	TArray<UObject*> OuterObjects;
+	StructPropertyHandle->GetOuterObjects(OuterObjects);
+
+	bool const bReadOnly = StructPropertyHandle->GetProperty()->HasAnyPropertyFlags( CPF_EditConst );
+
+	FText Title;
+	if (OuterObjects.Num() > 1)
+	{
+		FText const AssetName = FText::Format( LOCTEXT("GameplayTagDetailsBase_MultipleAssets", "{0} Assets"), FText::AsNumber( OuterObjects.Num() ) );
+		FText const PropertyName = StructPropertyHandle->GetPropertyDisplayName();
+		Title = FText::Format(LOCTEXT("GameplayTagQueryCustomization_BaseWidgetTitle", "Tag Editor: {0} {1}"), PropertyName, AssetName);
+	}
+	else if (OuterObjects.Num() > 0 && OuterObjects[0])
+	{
+		FText const AssetName = FText::FromString( OuterObjects[0]->GetName() );
+		FText const PropertyName = StructPropertyHandle->GetPropertyDisplayName();
+		Title = FText::Format(LOCTEXT("GameplayTagQueryCustomization_BaseWidgetTitle", "Tag Editor: {0} {1}"), PropertyName, AssetName);
+	}
+
+	GameplayTagQueryWidgetWindow = SNew(SWindow)
+	.Title(Title)
+	.HasCloseButton(false)
+	.ClientSize(FVector2D(600, 400))
+	[
+		SNew(SGameplayTagQueryWidget, EditableQueries)
+		.OnSaveAndClose(this, &FGameplayTagQueryCustomization::CloseWidgetWindow)
+		.OnCancel(this, &FGameplayTagQueryCustomization::CloseWidgetWindow)
+		.ReadOnly(bReadOnly)
+	];
+
+	IMainFrameModule& MainFrameModule = FModuleManager::LoadModuleChecked<IMainFrameModule>(TEXT("MainFrame"));
+	if (MainFrameModule.GetParentWindow().IsValid())
+	{
+		FSlateApplication::Get().AddWindowAsNativeChild(GameplayTagQueryWidgetWindow.ToSharedRef(), MainFrameModule.GetParentWindow().ToSharedRef());
+	}
+	else
+	{
+		FSlateApplication::Get().AddWindow(GameplayTagQueryWidgetWindow.ToSharedRef());
+	}
+
+	return FReply::Handled();
+}
+
+FGameplayTagQueryCustomization::~FGameplayTagQueryCustomization()
+{
+	if( GameplayTagQueryWidgetWindow.IsValid() )
+	{
+		GameplayTagQueryWidgetWindow->RequestDestroyWindow();
+	}
+	GEditor->UnregisterForUndo(this);
+}
+
+void FGameplayTagQueryCustomization::BuildEditableQueryList()
+{	
+	EditableQueries.Empty();
+
+	if( StructPropertyHandle.IsValid() )
+	{
+		TArray<void*> RawStructData;
+		StructPropertyHandle->AccessRawData(RawStructData);
+
+		TArray<UObject*> OuterObjects;
+		StructPropertyHandle->GetOuterObjects(OuterObjects);
+
+		ensure(RawStructData.Num() == OuterObjects.Num());
+		for (int32 Idx = 0; Idx < RawStructData.Num() && Idx < OuterObjects.Num(); ++Idx)
+		{
+			EditableQueries.Add(SGameplayTagQueryWidget::FEditableGameplayTagQueryDatum(OuterObjects[Idx], (FGameplayTagQuery*)RawStructData[Idx]));
+		}
+	}	
+}
+
+void FGameplayTagQueryCustomization::CloseWidgetWindow()
+{
+ 	if( GameplayTagQueryWidgetWindow.IsValid() )
+ 	{
+ 		GameplayTagQueryWidgetWindow->RequestDestroyWindow();
+ 	}
+}
+
+#undef LOCTEXT_NAMESPACE
diff --git a/Engine/Source/Editor/GameplayTagsEditor/Private/GameplayTagQueryCustomization.h b/Engine/Source/Editor/GameplayTagsEditor/Private/GameplayTagQueryCustomization.h
new file mode 100644
--- /dev/null
+++ b/Engine/Source/Editor/GameplayTagsEditor/Private/GameplayTagQueryCustomization.h
@@ -0,0 +1,46 @@
+// Copyright 1998-2015 Epic Games, Inc. All Rights Reserved.
+
+#pragma once
+
+#include "Editor/PropertyEditor/Public/PropertyEditorModule.h"
+#include "SGameplayTagQueryWidget.h"
+#include "EditorUndoClient.h"
+
+/** Customization for the gameplay tag query struct */
+class FGameplayTagQueryCustomization : public IPropertyTypeCustomization, public FEditorUndoClient
+{
+public:
+	static TSharedRef<IPropertyTypeCustomization> MakeInstance()
+	{
+		return MakeShareable(new FGameplayTagQueryCustomization);
+	}
+
+	~FGameplayTagQueryCustomization();
+
+	/** Overridden to show an edit button to launch the gameplay tag editor */
+	virtual void CustomizeHeader(TSharedRef<class IPropertyHandle> StructPropertyHandle, class FDetailWidgetRow& HeaderRow, IPropertyTypeCustomizationUtils& StructCustomizationUtils) override;
+	
+	virtual void CustomizeChildren(TSharedRef<IPropertyHandle> InStructPropertyHandle, class IDetailChildrenBuilder& ChildBuilder, IPropertyTypeCustomizationUtils& StructCustomizationUtils) override {}
+
+private:
+	/** Called when the edit button is clicked; Launches the gameplay tag editor */
+	FReply OnEditButtonClicked();
+
+	/** Overridden to do nothing */
+	FText GetEditButtonText() const;
+
+	void CloseWidgetWindow();
+
+	/** Build List of Editable Queries */
+	void BuildEditableQueryList();
+
+	/** Cached property handle */
+	TSharedPtr<IPropertyHandle> StructPropertyHandle;
+
+	/** The array of queries this objects has */
+	TArray<SGameplayTagQueryWidget::FEditableGameplayTagQueryDatum> EditableQueries;
+
+	/** The Window for the GameplayTagWidget */
+	TSharedPtr<SWindow> GameplayTagQueryWidgetWindow;
+};
+
diff --git a/Engine/Source/Editor/GameplayTagsEditor/Private/GameplayTagsEditorModule.cpp b/Engine/Source/Editor/GameplayTagsEditor/Private/GameplayTagsEditorModule.cpp
--- a/Engine/Source/Editor/GameplayTagsEditor/Private/GameplayTagsEditorModule.cpp
+++ b/Engine/Source/Editor/GameplayTagsEditor/Private/GameplayTagsEditorModule.cpp
@@ -4,6 +4,7 @@
 #include "GameplayTagsGraphPanelPinFactory.h"
 #include "GameplayTagsGraphPanelNodeFactory.h"
 #include "GameplayTagContainerCustomization.h"
+#include "GameplayTagQueryCustomization.h"
 #include "GameplayTagCustomization.h"
 #include "GameplayTagsSettings.h"
 #include "ISettingsModule.h"
@@ -26,6 +27,7 @@ public:
 		FPropertyEditorModule& PropertyModule = FModuleManager::LoadModuleChecked<FPropertyEditorModule>("PropertyEditor");
 		PropertyModule.RegisterCustomPropertyTypeLayout("GameplayTagContainer", FOnGetPropertyTypeCustomizationInstance::CreateStatic(&FGameplayTagContainerCustomization::MakeInstance));
 		PropertyModule.RegisterCustomPropertyTypeLayout("GameplayTag", FOnGetPropertyTypeCustomizationInstance::CreateStatic(&FGameplayTagCustomization::MakeInstance));
+		PropertyModule.RegisterCustomPropertyTypeLayout("GameplayTagQuery", FOnGetPropertyTypeCustomizationInstance::CreateStatic(&FGameplayTagQueryCustomization::MakeInstance));
 
 		TSharedPtr<FGameplayTagsGraphPanelPinFactory> GameplayTagsGraphPanelPinFactory = MakeShareable( new FGameplayTagsGraphPanelPinFactory() );
 		FEdGraphUtilities::RegisterVisualPinFactory(GameplayTagsGraphPanelPinFactory);
diff --git a/Engine/Source/Editor/GameplayTagsEditor/Private/SGameplayTagQueryWidget.cpp b/Engine/Source/Editor/GameplayTagsEditor/Private/SGameplayTagQueryWidget.cpp
new file mode 100644
--- /dev/null
+++ b/Engine/Source/Editor/GameplayTagsEditor/Private/SGameplayTagQueryWidget.cpp
@@ -0,0 +1,137 @@
+// Copyright 1998-2015 Epic Games, Inc. All Rights Reserved.
+
+#include "GameplayTagsEditorModulePrivatePCH.h"
+#include "SGameplayTagQueryWidget.h"
+#include "ScopedTransaction.h"
+#include "SScaleBox.h"
+
+#include "SGameplayTagWidget.h"
+#include "Editor/PropertyEditor/Public/PropertyEditorModule.h"
+#include "Editor/PropertyEditor/Public/IDetailsView.h"
+
+
+#define LOCTEXT_NAMESPACE "GameplayTagQueryWidget"
+
+void SGameplayTagQueryWidget::Construct(const FArguments& InArgs, const TArray<FEditableGameplayTagQueryDatum>& EditableTagQueries)
+{
+	ensure(EditableTagQueries.Num() > 0);
+	TagQueries = EditableTagQueries;
+
+	bReadOnly = InArgs._ReadOnly;
+	OnSaveAndClose = InArgs._OnSaveAndClose;
+	OnCancel = InArgs._OnCancel;
+
+	// Tag the assets as transactional so they can support undo/redo
+	for (int32 AssetIdx = 0; AssetIdx < TagQueries.Num(); ++AssetIdx)
+	{
+		UObject* TagQueryOwner = TagQueries[AssetIdx].TagQueryOwner.Get();
+		if (TagQueryOwner)
+		{
+			TagQueryOwner->SetFlags(RF_Transactional);
+		}
+	}
+
+	// build editable query object tree from the runtime query data
+	UEditableGameplayTagQuery* const EQ = CreateEditableQuery(*TagQueries[0].TagQuery);
+	EditableQuery = EQ;
+
+	// create details view for the editable query object
+	FDetailsViewArgs ViewArgs;
+	ViewArgs.bAllowSearch = false;
+	ViewArgs.bHideSelectionTip = true;
+	ViewArgs.bShowActorLabel = false;
+	
+	FPropertyEditorModule& PropertyModule = FModuleManager::LoadModuleChecked<FPropertyEditorModule>("PropertyEditor");
+	Details = PropertyModule.CreateDetailView(ViewArgs);
+	Details->SetObject(EQ);
+
+	ChildSlot
+	[
+		SNew(SScaleBox)
+		.HAlign(EHorizontalAlignment::HAlign_Left)
+		.VAlign(EVerticalAlignment::VAlign_Top)
+		.StretchDirection(EStretchDirection::DownOnly)
+		.Stretch(EStretch::ScaleToFit)
+		.Content()
+		[
+			SNew(SBorder)
+			.BorderImage(FEditorStyle::GetBrush("ToolPanel.GroupBorder"))
+			[
+				SNew(SVerticalBox)
+				+ SVerticalBox::Slot()
+				.AutoHeight()
+				.VAlign(VAlign_Top)
+				[
+					SNew(SHorizontalBox)
+					+ SHorizontalBox::Slot()
+					.AutoWidth()
+					[
+						SNew(SButton)
+						.IsEnabled(!bReadOnly)
+						.OnClicked(this, &SGameplayTagQueryWidget::OnSaveAndCloseClicked)
+						.Text(LOCTEXT("GameplayTagQueryWidget_SaveAndClose", "Save and Close"))
+					]
+					+ SHorizontalBox::Slot()
+					.AutoWidth()
+					[
+						SNew(SButton)
+						.OnClicked(this, &SGameplayTagQueryWidget::OnCancelClicked)
+						.Text(LOCTEXT("GameplayTagQueryWidget_Cancel", "Close Without Saving"))
+					]
+				]
+				// to delete!
+				+ SVerticalBox::Slot()
+				[
+					Details.ToSharedRef()
+				]
+			]
+		]
+	];
+}
+
+UEditableGameplayTagQuery* SGameplayTagQueryWidget::CreateEditableQuery(FGameplayTagQuery& Q)
+{
+	UEditableGameplayTagQuery* const EditableQuery = Q.CreateEditableQuery();
+	if (EditableQuery)
+	{
+		// prevent GC, will explicitly remove from root later
+		EditableQuery->AddToRoot();
+	}
+
+	return EditableQuery;
+}
+
+SGameplayTagQueryWidget::~SGameplayTagQueryWidget()
+{
+	// clean up our temp editing uobjects
+	UEditableGameplayTagQuery* const Q = EditableQuery.Get();
+	if (Q)
+	{
+		Q->RemoveFromRoot();
+	}
+}
+
+FReply SGameplayTagQueryWidget::OnSaveAndCloseClicked()
+{
+	// translate obj tree to token stream
+	if (EditableQuery.IsValid() && !bReadOnly)
+	{
+		// write to all selected queries
+		for (auto& TQ : TagQueries)
+		{
+			TQ.TagQuery->BuildFromEditableQuery(*EditableQuery.Get());
+			TQ.TagQueryOwner->MarkPackageDirty();
+		}
+	}
+
+	OnSaveAndClose.ExecuteIfBound();
+	return FReply::Handled();
+}
+
+FReply SGameplayTagQueryWidget::OnCancelClicked()
+{
+	OnCancel.ExecuteIfBound();
+	return FReply::Handled();
+}
+
+#undef LOCTEXT_NAMESPACE
diff --git a/Engine/Source/Editor/GameplayTagsEditor/Private/SGameplayTagQueryWidget.h b/Engine/Source/Editor/GameplayTagsEditor/Private/SGameplayTagQueryWidget.h
new file mode 100644
--- /dev/null
+++ b/Engine/Source/Editor/GameplayTagsEditor/Private/SGameplayTagQueryWidget.h
@@ -0,0 +1,72 @@
+// Copyright 1998-2015 Epic Games, Inc. All Rights Reserved.
+
+#pragma once
+
+#include "GameplayTagsManager.h"
+#include "SlateBasics.h"
+
+/** Widget allowing user to tag assets with gameplay tags */
+class SGameplayTagQueryWidget : public SCompoundWidget
+{
+public:
+
+	DECLARE_DELEGATE(FOnSaveAndClose)
+	DECLARE_DELEGATE(FOnCancel)
+	
+	SLATE_BEGIN_ARGS( SGameplayTagQueryWidget )
+	: _ReadOnly(false)
+	{}
+		SLATE_ARGUMENT( bool, ReadOnly ) // Flag to set if the list is read only
+		SLATE_EVENT(FOnSaveAndClose, OnSaveAndClose) // Called when "Save and Close" button clicked
+		SLATE_EVENT(FOnCancel, OnCancel) // Called when "Close Without Saving" button clicked
+	SLATE_END_ARGS()
+
+	/** Simple struct holding a tag query and its owner for generic re-use of the widget */
+	struct FEditableGameplayTagQueryDatum
+	{
+		/** Constructor */
+		FEditableGameplayTagQueryDatum(class UObject* InOwnerObj, struct FGameplayTagQuery* InTagQuery)
+			: TagQueryOwner(InOwnerObj)
+			, TagQuery(InTagQuery)
+		{}
+
+		/** Owning UObject of the query being edited */
+		TWeakObjectPtr<class UObject> TagQueryOwner;
+
+		/** Tag query to edit */
+		struct FGameplayTagQuery* TagQuery; 
+	};
+
+	/** Construct the actual widget */
+	void Construct(const FArguments& InArgs, const TArray<FEditableGameplayTagQueryDatum>& EditableTagQueries);
+
+	~SGameplayTagQueryWidget();
+
+private:
+
+	/* Flag to set if the list is read only*/
+	bool bReadOnly;
+
+	/** Containers to modify */
+	TArray<FEditableGameplayTagQueryDatum> TagQueries;
+
+	/** Called when "save and close" is clicked */
+	FOnSaveAndClose OnSaveAndClose;
+
+	/** Called when "cancel" is clicked */
+	FOnCancel OnCancel;
+
+	/** Properties Tab */
+	TSharedPtr<class IDetailsView> Details;
+
+	class UEditableGameplayTagQuery* CreateEditableQuery(FGameplayTagQuery& Q);
+	TWeakObjectPtr<class UEditableGameplayTagQuery> EditableQuery;
+
+	/** Called when the user clicks the "Expand All" button; Expands the entire tag tree */
+	FReply OnSaveAndCloseClicked();
+
+	/** Called when the user clicks the "Expand All" button; Expands the entire tag tree */
+	FReply OnCancelClicked();
+
+};
+
diff --git a/Engine/Source/Runtime/GameplayTags/Classes/GameplayTagContainer.h b/Engine/Source/Runtime/GameplayTags/Classes/GameplayTagContainer.h
--- a/Engine/Source/Runtime/GameplayTags/Classes/GameplayTagContainer.h
+++ b/Engine/Source/Runtime/GameplayTags/Classes/GameplayTagContainer.h
@@ -168,6 +168,15 @@ struct GAMEPLAYTAGS_API FGameplayTagContainer
 	*/
 	bool MatchesAll(const FGameplayTagContainer& Other, bool bCountEmptyAsMatch) const;
 
+	/** 
+	 * Checks if this container matches the given query.
+	 *
+	 * @param Query		Query we are checking against
+	 *
+	 * @return True if this container matches the query, false otherwise.
+	 */
+	bool MatchesQuery(const struct FGameplayTagQuery& Query) const;
+
 	/**
 	 * Determine if the container has the specified tag
 	 * 
@@ -330,3 +339,304 @@ struct TStructOpsTypeTraits<FGameplayTagContainer> : public TStructOpsTypeTraits
 		WithCopy = true
 	};
 };
+
+
+/** Enumerates the list of supported query expression types. */
+UENUM()
+namespace EGameplayTagQueryExprType
+{
+	enum Type
+	{
+		Undefined = 0,
+		AnyTagsMatch,
+		AllTagsMatch,
+		NoTagsMatch,
+		AnyExprMatch,
+		AllExprMatch,
+		NoExprMatch,
+	};
+}
+
+namespace EGameplayTagQueryStreamVersion
+{
+	enum Type
+	{
+		InitialVersion = 0,
+
+		// -----<new versions can be added before this line>-------------------------------------------------
+		// - this needs to be the last line (see note below)
+		VersionPlusOne,
+		LatestVersion = VersionPlusOne - 1
+	};
+}
+
+/**
+ * An FGameplayTagQuery is a logical query that can be run against an FGameplayTagContainer.  A query that succeeds is said to "match".
+ * Queries are logical expressions that can test the intersection properties of another tag container (all, any, or none), or the matching state of a set of sub-expressions
+ * (all, any, or none). This allows queries to be arbitrarily recursive and very expressive.  For instance, if you wanted to test if a given tag container contained tags 
+ * ((A && B) || (C)) && (!D), you would construct your query in the form ALL( ANY( ALL(A,B), ALL(C) ), NONE(D) )
+ * 
+ * You can expose the query structs to Blueprints and edit them with a custom editor, or you can construct them natively in code. 
+ * 
+ * Example of how to build a query via code:
+ *	FGameplayTagQuery Q;
+ *	Q.BuildQuery(
+ *		FGameplayTagQueryExpression()
+ * 		.AllTagsMatch()
+ *		.AddTag(IGameplayTagsModule::RequestGameplayTag(FName(TEXT("Animal.Mammal.Dog.Corgi"))))
+ *		.AddTag(IGameplayTagsModule::RequestGameplayTag(FName(TEXT("Plant.Tree.Spruce"))))
+ *		);
+ * 
+ * Queries are internally represented as a byte stream that is memory-efficient and can be evaluated quickly at runtime.
+ */
+USTRUCT(BlueprintType)
+struct GAMEPLAYTAGS_API FGameplayTagQuery
+{
+	GENERATED_BODY();
+
+public:
+	FGameplayTagQuery();
+
+private:
+	/** Versioning for future token stream protocol changes. See EGameplayTagQueryStreamVersion. */
+	UPROPERTY()
+	int32 TokenStreamVersion;
+
+	/** List of tags referenced by this entire query. Token stream stored indices into this list. */
+	UPROPERTY()
+	TArray<FGameplayTag> TagDictionary;
+
+	/** Stream representation of the actual hierarchical query */
+	UPROPERTY()
+	TArray<uint8> QueryTokenStream;
+
+	/** Returns a gameplay tag from the tag dictionary */
+	FGameplayTag GetTagFromIndex(int32 TagIdx) const
+	{
+		ensure(TagDictionary.IsValidIndex(TagIdx));
+		return TagDictionary[TagIdx];
+	}
+
+public:
+	/** Returns true if the given tags match this query, or false otherwise. */
+	bool Matches(FGameplayTagContainer const& Tags) const;
+
+	/** Creates this query with the given root expression. */
+	void BuildQuery(struct FGameplayTagQueryExpression& RootQueryExpr);
+
+	/** Builds a FGameplayTagQueryExpression from this query. */
+	void GetQueryExpr(struct FGameplayTagQueryExpression& OutExpr) const;
+	
+#if WITH_EDITOR
+	/** Creates this query based on the given EditableQuery object */
+	void BuildFromEditableQuery(class UEditableGameplayTagQuery& EditableQuery); 
+
+	/** Creates editable query object tree based on this query */
+	UEditableGameplayTagQuery* CreateEditableQuery();
+#endif // WITH_EDITOR
+
+	static const FGameplayTagQuery EmptyQuery;
+
+	friend class FQueryEvaluator;
+};
+
+struct GAMEPLAYTAGS_API FGameplayTagQueryExpression
+{
+	FGameplayTagQueryExpression& AnyTagsMatch()
+	{
+		ExprType = EGameplayTagQueryExprType::AnyTagsMatch;
+		return *this;
+	}
+
+	FGameplayTagQueryExpression& AllTagsMatch()
+	{
+		ExprType = EGameplayTagQueryExprType::AllTagsMatch;
+		return *this;
+	}
+
+	FGameplayTagQueryExpression& NoTagsMatch()
+	{
+		ExprType = EGameplayTagQueryExprType::NoTagsMatch;
+		return *this;
+	}
+
+	FGameplayTagQueryExpression& AnyExprMatch()
+	{
+		ExprType = EGameplayTagQueryExprType::AnyExprMatch;
+		return *this;
+	}
+
+	FGameplayTagQueryExpression& AllExprMatch()
+	{
+		ExprType = EGameplayTagQueryExprType::AllExprMatch;
+		return *this;
+	}
+
+	FGameplayTagQueryExpression& NoExprMatch()
+	{
+		ExprType = EGameplayTagQueryExprType::NoExprMatch;
+		return *this;
+	}
+
+	FGameplayTagQueryExpression& AddTag(FGameplayTag Tag)
+	{
+		ensure(UsesTagSet());
+		TagSet.Add(Tag);
+		return *this;
+	}
+
+	FGameplayTagQueryExpression& AddTags(FGameplayTagContainer const& Tags)
+	{
+		ensure(UsesTagSet());
+		for (auto T : Tags)
+		{
+			TagSet.Add(T);
+		}
+		return *this;
+	}
+
+	FGameplayTagQueryExpression& AddExpr(FGameplayTagQueryExpression& Expr)
+	{
+		ensure(UsesExprSet());
+		ExprSet.Add(Expr);
+		return *this;
+	}
+	
+	/** Writes this expression to the given token stream. */
+	void EmitTokens(TArray<uint8>& TokenStream, TArray<FGameplayTag>& TagDictionary) const;
+
+	/** Which type of expression this is. */
+	EGameplayTagQueryExprType::Type ExprType;
+	/** Expression list, for expression types that need it */
+	TArray<struct FGameplayTagQueryExpression> ExprSet;
+	/** Tag list, for expression types that need it */
+	TArray<FGameplayTag> TagSet;
+
+	/** Returns true if this expression uses the tag data. */
+	FORCEINLINE bool UsesTagSet() const
+	{
+		return (ExprType == EGameplayTagQueryExprType::AllTagsMatch) || (ExprType == EGameplayTagQueryExprType::AnyTagsMatch) || (ExprType == EGameplayTagQueryExprType::NoTagsMatch);
+	}
+	/** Returns true if this expression uses the expression list data. */
+	FORCEINLINE bool UsesExprSet() const
+	{
+		return (ExprType == EGameplayTagQueryExprType::AllExprMatch) || (ExprType == EGameplayTagQueryExprType::AnyExprMatch) || (ExprType == EGameplayTagQueryExprType::NoExprMatch);
+	}
+};
+
+/** 
+ * This is an editor-only representation of a query, designed to be editable with a typical property window. 
+ * To edit a query in the editor, an FGameplayTagQuery is converted to a set of UObjects and edited,  When finished,
+ * the query struct is rewritten and these UObjects are discarded.
+ * This query representation is not intended for runtime use.
+ */
+UCLASS(editinlinenew, collapseCategories, Transient) 
+class GAMEPLAYTAGS_API UEditableGameplayTagQuery : public UObject
+{
+	GENERATED_BODY()
+
+public:
+	/** The base expression of this query. */
+	UPROPERTY(EditDefaultsOnly, Instanced, Category = Query)
+	class UEditableGameplayTagQueryExpression* RootExpression;
+
+#if WITH_EDITOR
+	/** Converts this editor query construct into the runtime-usable token stream. */
+	void EmitTokens(TArray<uint8>& TokenStream, TArray<FGameplayTag>& TagDictionary) const;
+#endif  // WITH_EDITOR
+};
+
+UCLASS(abstract, editinlinenew, collapseCategories, Transient)
+class GAMEPLAYTAGS_API UEditableGameplayTagQueryExpression : public UObject
+{
+	GENERATED_BODY()
+
+#if WITH_EDITOR
+public:
+	/** Converts this editor query construct into the runtime-usable token stream. */
+	virtual void EmitTokens(TArray<uint8>& TokenStream, TArray<FGameplayTag>& TagDictionary) const {};
+
+protected:
+ 	void EmitTagTokens(FGameplayTagContainer const& TagsToEmit, TArray<uint8>& TokenStream, TArray<FGameplayTag>& TagDictionary) const;
+ 	void EmitExprListTokens(TArray<UEditableGameplayTagQueryExpression*> const& ExprList, TArray<uint8>& TokenStream, TArray<FGameplayTag>& TagDictionary) const;
+#endif  // WITH_EDITOR
+};
+
+UCLASS(BlueprintType, editinlinenew, collapseCategories, meta=(DisplayName="Any Tags Match"))
+class UEditableGameplayTagQueryExpression_AnyTagsMatch : public UEditableGameplayTagQueryExpression
+{
+	GENERATED_BODY()
+public:
+	UPROPERTY(EditDefaultsOnly, Category = Expr)
+	FGameplayTagContainer Tags;
+
+#if WITH_EDITOR
+	virtual void EmitTokens(TArray<uint8>& TokenStream, TArray<FGameplayTag>& TagDictionary) const override;
+#endif  // WITH_EDITOR
+};
+
+UCLASS(BlueprintType, editinlinenew, collapseCategories, meta = (DisplayName = "All Tags Match"))
+class UEditableGameplayTagQueryExpression_AllTagsMatch : public UEditableGameplayTagQueryExpression
+{
+	GENERATED_BODY()
+public:
+	UPROPERTY(EditDefaultsOnly, Category = Expr)
+	FGameplayTagContainer Tags;
+
+#if WITH_EDITOR
+	virtual void EmitTokens(TArray<uint8>& TokenStream, TArray<FGameplayTag>& TagDictionary) const override;
+#endif  // WITH_EDITOR
+};
+
+UCLASS(BlueprintType, editinlinenew, collapseCategories, meta = (DisplayName = "No Tags Match"))
+class UEditableGameplayTagQueryExpression_NoTagsMatch : public UEditableGameplayTagQueryExpression
+{
+	GENERATED_BODY()
+public:
+	UPROPERTY(EditDefaultsOnly, Category = Expr)
+	FGameplayTagContainer Tags;
+
+#if WITH_EDITOR
+	virtual void EmitTokens(TArray<uint8>& TokenStream, TArray<FGameplayTag>& TagDictionary) const override;
+#endif  // WITH_EDITOR
+};
+
+UCLASS(BlueprintType, editinlinenew, collapseCategories, meta = (DisplayName = "Any Expressions Match"))
+class UEditableGameplayTagQueryExpression_AnyExprMatch : public UEditableGameplayTagQueryExpression
+{
+	GENERATED_BODY()
+public:
+	UPROPERTY(EditAnywhere, Instanced, Category = Expr)
+	TArray<UEditableGameplayTagQueryExpression*> Expressions;
+
+#if WITH_EDITOR
+	virtual void EmitTokens(TArray<uint8>& TokenStream, TArray<FGameplayTag>& TagDictionary) const override;
+#endif  // WITH_EDITOR
+};
+
+UCLASS(BlueprintType, editinlinenew, collapseCategories, meta = (DisplayName = "All Expressions Match"))
+class UEditableGameplayTagQueryExpression_AllExprMatch : public UEditableGameplayTagQueryExpression
+{
+	GENERATED_BODY()
+public:
+	UPROPERTY(EditAnywhere, Instanced, Category = Expr)
+	TArray<UEditableGameplayTagQueryExpression*> Expressions;
+
+#if WITH_EDITOR
+	virtual void EmitTokens(TArray<uint8>& TokenStream, TArray<FGameplayTag>& TagDictionary) const override;
+#endif  // WITH_EDITOR
+};
+
+UCLASS(BlueprintType, editinlinenew, collapseCategories, meta = (DisplayName = "No Expressions Match"))
+class UEditableGameplayTagQueryExpression_NoExprMatch : public UEditableGameplayTagQueryExpression
+{
+	GENERATED_BODY()
+public:
+	UPROPERTY(EditAnywhere, Instanced, Category = Expr)
+	TArray<UEditableGameplayTagQueryExpression*> Expressions;
+
+#if WITH_EDITOR
+	virtual void EmitTokens(TArray<uint8>& TokenStream, TArray<FGameplayTag>& TagDictionary) const override;
+#endif  // WITH_EDITOR
+};
+
diff --git a/Engine/Source/Runtime/GameplayTags/Private/GameplayTagContainer.cpp b/Engine/Source/Runtime/GameplayTags/Private/GameplayTagContainer.cpp
--- a/Engine/Source/Runtime/GameplayTags/Private/GameplayTagContainer.cpp
+++ b/Engine/Source/Runtime/GameplayTags/Private/GameplayTagContainer.cpp
@@ -3,6 +3,392 @@
 #include "GameplayTagsModulePrivatePCH.h"
 
 const FGameplayTagContainer FGameplayTagContainer::EmptyContainer;
+const FGameplayTagQuery FGameplayTagQuery::EmptyQuery;
+
+
+/** Helper class to parse/eval query token streams. */
+class FQueryEvaluator
+{
+public:
+	FQueryEvaluator(FGameplayTagQuery const& Q)
+		: Query(Q), 
+		CurStreamIdx(0), 
+		Version(EGameplayTagQueryStreamVersion::LatestVersion), 
+		bReadError(false)
+	{}
+
+	/** Evaluates the query against the given tag container and returns the result (true if matching, false otherwise). */
+	bool Eval(FGameplayTagContainer const& Tags);
+
+	/** Parses the token stream into an FExpr. */
+	void Read(struct FGameplayTagQueryExpression& E);
+
+private:
+	FGameplayTagQuery const& Query;
+	int32 CurStreamIdx;
+	int32 Version;
+	bool bReadError;
+
+	bool EvalAnyTagsMatch(FGameplayTagContainer const& Tags, bool bSkip);
+	bool EvalAllTagsMatch(FGameplayTagContainer const& Tags, bool bSkip);
+	bool EvalNoTagsMatch(FGameplayTagContainer const& Tags, bool bSkip);
+
+	bool EvalAnyExprMatch(FGameplayTagContainer const& Tags, bool bSkip);
+	bool EvalAllExprMatch(FGameplayTagContainer const& Tags, bool bSkip);
+	bool EvalNoExprMatch(FGameplayTagContainer const& Tags, bool bSkip);
+
+	bool EvalExpr(FGameplayTagContainer const& Tags, bool bSkip = false);
+	void ReadExpr(struct FGameplayTagQueryExpression& E);
+
+#if WITH_EDITOR
+public:
+	UEditableGameplayTagQuery* CreateEditableQuery();
+
+private:
+	UEditableGameplayTagQueryExpression* ReadEditableQueryExpr(UObject* ExprOuter);
+	void ReadEditableQueryTags(UEditableGameplayTagQueryExpression* EditableQueryExpr);
+	void ReadEditableQueryExprList(UEditableGameplayTagQueryExpression* EditableQueryExpr);
+#endif // WITH_EDITOR
+
+	/** Returns the next token in the stream. If there's a read error, sets bReadError and returns zero, so be sure to check that. */
+	uint8 GetToken()
+	{
+		if (Query.QueryTokenStream.IsValidIndex(CurStreamIdx))
+		{
+			return Query.QueryTokenStream[CurStreamIdx++];
+		}
+		
+		UE_LOG(LogGameplayTags, Warning, TEXT("Error parsing FGameplayTagQuery!"));
+		bReadError = true;
+		return 0;
+	}
+};
+
+bool FQueryEvaluator::Eval(FGameplayTagContainer const& Tags)
+{
+	CurStreamIdx = 0;
+
+	// start parsing the set
+	Version = GetToken();
+	if (bReadError)
+	{
+		return false;
+	}
+
+	bool const bRet = EvalExpr(Tags);
+	ensure(CurStreamIdx == Query.QueryTokenStream.Num());
+	return bRet;
+}
+
+void FQueryEvaluator::Read(FGameplayTagQueryExpression& E)
+{
+	E = FGameplayTagQueryExpression();
+	CurStreamIdx = 0;
+
+	if (Query.QueryTokenStream.Num() > 0)
+	{
+		// start parsing the set
+		Version = GetToken();
+		if (!bReadError)
+		{
+			uint8 const bHasRootExpression = GetToken();
+			if (!bReadError && bHasRootExpression)
+			{
+				ReadExpr(E);
+			}
+		}
+
+		ensure(CurStreamIdx == Query.QueryTokenStream.Num());
+	}
+}
+
+void FQueryEvaluator::ReadExpr(FGameplayTagQueryExpression& E)
+{
+	E.ExprType = (EGameplayTagQueryExprType::Type) GetToken();
+	if (bReadError)
+	{
+		return;
+	}
+	
+	if (E.UsesTagSet())
+	{
+		// parse tag set
+		int32 NumTags = GetToken();
+		if (bReadError)
+		{
+			return;
+		}
+
+		for (int32 Idx = 0; Idx < NumTags; ++Idx)
+		{
+			int32 const TagIdx = GetToken();
+			if (bReadError)
+			{
+				return;
+			}
+
+			FGameplayTag Tag = Query.GetTagFromIndex(TagIdx);
+			E.AddTag(Tag);
+		}
+	}
+	else
+	{
+		// parse expr set
+		int32 NumExprs = GetToken();
+		if (bReadError)
+		{
+			return;
+		}
+
+		for (int32 Idx = 0; Idx < NumExprs; ++Idx)
+		{
+			FGameplayTagQueryExpression E;
+			ReadExpr(E);
+			E.AddExpr(E);
+		}
+	}
+}
+
+
+bool FQueryEvaluator::EvalAnyTagsMatch(FGameplayTagContainer const& Tags, bool bSkip)
+{
+	bool bShortCircuit = bSkip;
+	bool Result = false;
+
+	// parse tagset
+	int32 const NumTags = GetToken();
+	if (bReadError)
+	{
+		return false;
+	}
+
+	for (int32 Idx = 0; Idx < NumTags; ++Idx)
+	{
+		int32 const TagIdx = GetToken();
+		if (bReadError)
+		{
+			return false;
+		}
+
+		if (bShortCircuit == false)
+		{
+			FGameplayTag Tag = Query.GetTagFromIndex(TagIdx);
+
+			bool bHasTag = Tags.HasTag(Tag, EGameplayTagMatchType::IncludeParentTags, EGameplayTagMatchType::Explicit);
+
+			if (bHasTag)
+			{
+				// one match is sufficient for a true result!
+				bShortCircuit = true;
+				Result = true;
+			}
+		}
+	}
+
+	return Result;
+}
+
+bool FQueryEvaluator::EvalAllTagsMatch(FGameplayTagContainer const& Tags, bool bSkip)
+{
+	bool bShortCircuit = bSkip;
+
+	// assume true until proven otherwise
+	bool Result = true;
+
+	// parse tagset
+	int32 const NumTags = GetToken();
+	if (bReadError)
+	{
+		return false;
+	}
+
+	for (int32 Idx = 0; Idx < NumTags; ++Idx)
+	{
+		int32 const TagIdx = GetToken();
+		if (bReadError)
+		{
+			return false;
+		}
+
+		if (bShortCircuit == false)
+		{
+			FGameplayTag const Tag = Query.GetTagFromIndex(TagIdx);
+			bool const bHasTag = Tags.HasTag(Tag, EGameplayTagMatchType::IncludeParentTags, EGameplayTagMatchType::Explicit);
+
+			if (bHasTag == false)
+			{
+				// one failed match is sufficient for a false result
+				bShortCircuit = true;
+				Result = false;
+			}
+		}
+	}
+
+	return Result;
+}
+
+bool FQueryEvaluator::EvalNoTagsMatch(FGameplayTagContainer const& Tags, bool bSkip)
+{
+	bool bShortCircuit = bSkip;
+
+	// assume true until proven otherwise
+	bool Result = true;
+
+	// parse tagset
+	int32 const NumTags = GetToken();
+	if (bReadError)
+	{
+		return false;
+	}
+	
+	for (int32 Idx = 0; Idx < NumTags; ++Idx)
+	{
+		int32 const TagIdx = GetToken();
+		if (bReadError)
+		{
+			return false;
+		}
+
+		if (bShortCircuit == false)
+		{
+			FGameplayTag const Tag = Query.GetTagFromIndex(TagIdx);
+			bool const bHasTag = Tags.HasTag(Tag, EGameplayTagMatchType::IncludeParentTags, EGameplayTagMatchType::Explicit);
+
+			if (bHasTag == true)
+			{
+				// one match is sufficient for a false result
+				bShortCircuit = true;
+				Result = false;
+			}
+		}
+	}
+
+	return Result;
+}
+
+bool FQueryEvaluator::EvalAnyExprMatch(FGameplayTagContainer const& Tags, bool bSkip)
+{
+	bool bShortCircuit = bSkip;
+
+	// assume false until proven otherwise
+	bool Result = false;
+
+	// parse exprset
+	int32 const NumExprs = GetToken();
+	if (bReadError)
+	{
+		return false;
+	}
+
+	for (int32 Idx = 0; Idx < NumExprs; ++Idx)
+	{
+		bool const bExprResult = EvalExpr(Tags, bShortCircuit);
+		if (bShortCircuit == false)
+		{
+			if (bExprResult == true)
+			{
+				// one match is sufficient for true result
+				Result = true;
+				bShortCircuit = true;
+			}
+		}
+	}
+
+	return Result;
+}
+bool FQueryEvaluator::EvalAllExprMatch(FGameplayTagContainer const& Tags, bool bSkip)
+{
+	bool bShortCircuit = bSkip;
+
+	// assume false until proven otherwise
+	bool Result = false;
+
+	// parse exprset
+	int32 const NumExprs = GetToken();
+	if (bReadError)
+	{
+		return false;
+	}
+
+	for (int32 Idx = 0; Idx < NumExprs; ++Idx)
+	{
+		bool const bExprResult = EvalExpr(Tags, bShortCircuit);
+		if (bShortCircuit == false)
+		{
+			if (bExprResult == true)
+			{
+				// one match is sufficient for true result
+				Result = true;
+				bShortCircuit = true;
+			}
+		}
+	}
+
+	return Result;
+}
+bool FQueryEvaluator::EvalNoExprMatch(FGameplayTagContainer const& Tags, bool bSkip)
+{
+	bool bShortCircuit = bSkip;
+
+	// assume false until proven otherwise
+	bool Result = false;
+
+	// parse exprset
+	int32 const NumExprs = GetToken();
+	if (bReadError)
+	{
+		return false;
+	}
+
+	for (int32 Idx = 0; Idx < NumExprs; ++Idx)
+	{
+		bool const bExprResult = EvalExpr(Tags, bShortCircuit);
+		if (bShortCircuit == false)
+		{
+			if (bExprResult == true)
+			{
+				// one match is sufficient for true result
+				Result = true;
+				bShortCircuit = true;
+			}
+		}
+	}
+
+	return Result;
+}
+
+
+bool FQueryEvaluator::EvalExpr(FGameplayTagContainer const& Tags, bool bSkip)
+{
+	EGameplayTagQueryExprType::Type const ExprType = (EGameplayTagQueryExprType::Type) GetToken();
+	if (bReadError)
+	{
+		return false;
+	}
+
+	// emit exprdata
+	switch (ExprType)
+	{
+	case EGameplayTagQueryExprType::AnyTagsMatch:
+		return EvalAnyTagsMatch(Tags, bSkip);
+	case EGameplayTagQueryExprType::AllTagsMatch:
+		return EvalAllTagsMatch(Tags, bSkip);
+	case EGameplayTagQueryExprType::NoTagsMatch:
+		return EvalNoTagsMatch(Tags, bSkip);
+
+	case EGameplayTagQueryExprType::AnyExprMatch:
+		return EvalAnyExprMatch(Tags, bSkip);
+	case EGameplayTagQueryExprType::AllExprMatch:
+		return EvalAllExprMatch(Tags, bSkip);
+	case EGameplayTagQueryExprType::NoExprMatch:
+		return EvalNoExprMatch(Tags, bSkip);
+	}
+
+	check(false);
+	return false;
+}
+
+
 
 FGameplayTagContainer::FGameplayTagContainer()
 {}
@@ -177,6 +563,11 @@ bool FGameplayTagContainer::MatchesAny(FGameplayTagContainer const& Other, bool
 	return DoesTagContainerMatch(Other, EGameplayTagMatchType::IncludeParentTags, EGameplayTagMatchType::Explicit, EGameplayContainerMatchType::Any);
 }
 
+bool FGameplayTagContainer::MatchesQuery(const FGameplayTagQuery& Query) const
+{
+	return Query.Matches(*this);
+}
+
 void FGameplayTagContainer::AppendTags(FGameplayTagContainer const& Other)
 {
 	//add all the tags
@@ -417,4 +808,341 @@ bool FGameplayTag::NetSerialize(FArchive& Ar, class UPackageMap* Map, bool& bOut
 
 	bOutSuccess = true;
 	return true;
-}
\ No newline at end of file
+}
+
+
+
+FGameplayTagQuery::FGameplayTagQuery()
+	: TokenStreamVersion(EGameplayTagQueryStreamVersion::LatestVersion)
+{
+}
+
+bool FGameplayTagQuery::Matches(FGameplayTagContainer const& Tags) const
+{
+	FQueryEvaluator QE(*this);
+	return QE.Eval(Tags);
+}
+
+void FGameplayTagQuery::GetQueryExpr(FGameplayTagQueryExpression& OutExpr) const
+{
+	// build the FExpr tree from the token stream and return it
+	FQueryEvaluator QE(*this);
+	QE.Read(OutExpr);
+}
+
+void FGameplayTagQuery::BuildQuery(FGameplayTagQueryExpression& QueryExpr)
+{
+	// Reserve size here is arbitrary, goal is to minimizing reallocs while being respectful of mem usage
+	QueryTokenStream.Reset(128);
+	TagDictionary.Reset();
+
+	// add stream version first
+	QueryTokenStream.Add(EGameplayTagQueryStreamVersion::LatestVersion);
+
+	// emit the query
+	QueryTokenStream.Add(1);		// true to indicate is has a root expression
+	QueryExpr.EmitTokens(QueryTokenStream, TagDictionary);
+}
+
+
+#if WITH_EDITOR
+
+UEditableGameplayTagQuery* FQueryEvaluator::CreateEditableQuery()
+{
+	CurStreamIdx = 0;
+
+	UEditableGameplayTagQuery* const EditableQuery = NewObject<UEditableGameplayTagQuery>(GetTransientPackage(), NAME_None, RF_Transactional);
+
+	// start parsing the set
+	Version = GetToken();
+	if (!bReadError)
+	{
+		uint8 const bHasRootExpression = GetToken();
+		if (!bReadError && bHasRootExpression)
+		{
+			EditableQuery->RootExpression = ReadEditableQueryExpr(EditableQuery);
+		}
+	}
+	ensure(CurStreamIdx == Query.QueryTokenStream.Num());
+
+	return EditableQuery;
+}
+
+UEditableGameplayTagQueryExpression* FQueryEvaluator::ReadEditableQueryExpr(UObject* ExprOuter)
+{
+	EGameplayTagQueryExprType::Type const ExprType = (EGameplayTagQueryExprType::Type) GetToken();
+	if (bReadError)
+	{
+		return nullptr;
+	}
+
+	UClass* ExprClass = nullptr;
+	switch (ExprType)
+	{
+	case EGameplayTagQueryExprType::AnyTagsMatch:
+		ExprClass = UEditableGameplayTagQueryExpression_AnyTagsMatch::StaticClass();
+		break;
+	case EGameplayTagQueryExprType::AllTagsMatch:
+		ExprClass = UEditableGameplayTagQueryExpression_AllTagsMatch::StaticClass();
+		break;
+	case EGameplayTagQueryExprType::NoTagsMatch:
+		ExprClass = UEditableGameplayTagQueryExpression_NoTagsMatch::StaticClass();
+		break;
+	case EGameplayTagQueryExprType::AnyExprMatch:
+		ExprClass = UEditableGameplayTagQueryExpression_AnyExprMatch::StaticClass();
+		break;
+	case EGameplayTagQueryExprType::AllExprMatch:
+		ExprClass = UEditableGameplayTagQueryExpression_AllExprMatch::StaticClass();
+		break;
+	case EGameplayTagQueryExprType::NoExprMatch:
+		ExprClass = UEditableGameplayTagQueryExpression_NoExprMatch::StaticClass();
+		break;
+	}
+
+	UEditableGameplayTagQueryExpression* NewExpr = nullptr;
+	if (ExprClass)
+	{
+		NewExpr = NewObject<UEditableGameplayTagQueryExpression>(ExprOuter, ExprClass, NAME_None, RF_Transactional);
+		if (NewExpr)
+		{
+			switch (ExprType)
+			{
+			case EGameplayTagQueryExprType::AnyTagsMatch:
+			case EGameplayTagQueryExprType::AllTagsMatch:
+			case EGameplayTagQueryExprType::NoTagsMatch:
+				ReadEditableQueryTags(NewExpr);
+				break;
+			case EGameplayTagQueryExprType::AnyExprMatch:
+			case EGameplayTagQueryExprType::AllExprMatch:
+			case EGameplayTagQueryExprType::NoExprMatch:
+				ReadEditableQueryExprList(NewExpr);
+				break;
+			}
+		}
+	}
+
+	return NewExpr;
+}
+
+void FQueryEvaluator::ReadEditableQueryTags(UEditableGameplayTagQueryExpression* EditableQueryExpr)
+{
+	// find the tag container to read into
+	FGameplayTagContainer* Tags = nullptr;
+	if (EditableQueryExpr->IsA(UEditableGameplayTagQueryExpression_AnyTagsMatch::StaticClass()))
+	{
+		Tags = &((UEditableGameplayTagQueryExpression_AnyTagsMatch*)EditableQueryExpr)->Tags;
+	}
+	else if (EditableQueryExpr->IsA(UEditableGameplayTagQueryExpression_AllTagsMatch::StaticClass()))
+	{
+		Tags = &((UEditableGameplayTagQueryExpression_AllTagsMatch*)EditableQueryExpr)->Tags;
+	}
+	else if (EditableQueryExpr->IsA(UEditableGameplayTagQueryExpression_NoTagsMatch::StaticClass()))
+	{
+		Tags = &((UEditableGameplayTagQueryExpression_NoTagsMatch*)EditableQueryExpr)->Tags;
+	}
+	ensure(Tags);
+
+	if (Tags)
+	{
+		// parse tag set
+		int32 const NumTags = GetToken();
+		if (bReadError)
+		{
+			return;
+		}
+
+		for (int32 Idx = 0; Idx < NumTags; ++Idx)
+		{
+			int32 const TagIdx = GetToken();
+			if (bReadError)
+			{
+				return;
+			}
+
+			FGameplayTag const Tag = Query.GetTagFromIndex(TagIdx);
+			Tags->AddTag(Tag);
+		}
+	}
+}
+
+void FQueryEvaluator::ReadEditableQueryExprList(UEditableGameplayTagQueryExpression* EditableQueryExpr)
+{
+	// find the tag container to read into
+	TArray<UEditableGameplayTagQueryExpression*>* ExprList = nullptr;
+	if (EditableQueryExpr->IsA(UEditableGameplayTagQueryExpression_AnyExprMatch::StaticClass()))
+	{
+		ExprList = &((UEditableGameplayTagQueryExpression_AnyExprMatch*)EditableQueryExpr)->Expressions;
+	}
+	else if (EditableQueryExpr->IsA(UEditableGameplayTagQueryExpression_AllExprMatch::StaticClass()))
+	{
+		ExprList = &((UEditableGameplayTagQueryExpression_AllExprMatch*)EditableQueryExpr)->Expressions;
+	}
+	else if (EditableQueryExpr->IsA(UEditableGameplayTagQueryExpression_NoExprMatch::StaticClass()))
+	{
+		ExprList = &((UEditableGameplayTagQueryExpression_NoExprMatch*)EditableQueryExpr)->Expressions;
+	}
+	ensure(ExprList);
+
+	if (ExprList)
+	{
+		// parse expr set
+		int32 const NumExprs = GetToken();
+		if (bReadError)
+		{
+			return;
+		}
+
+		for (int32 Idx = 0; Idx < NumExprs; ++Idx)
+		{
+			UEditableGameplayTagQueryExpression* const NewExpr = ReadEditableQueryExpr(EditableQueryExpr);
+			ExprList->Add(NewExpr);
+		}
+	}
+}
+
+UEditableGameplayTagQuery* FGameplayTagQuery::CreateEditableQuery()
+{
+	FQueryEvaluator QE(*this);
+	return QE.CreateEditableQuery();
+}
+
+void FGameplayTagQuery::BuildFromEditableQuery(UEditableGameplayTagQuery& EditableQuery)
+{
+	QueryTokenStream.Reset();
+	TagDictionary.Reset();
+
+	// add stream version first
+	QueryTokenStream.Add(EGameplayTagQueryStreamVersion::LatestVersion);
+
+	EditableQuery.EmitTokens(QueryTokenStream, TagDictionary);
+}
+
+void UEditableGameplayTagQuery::EmitTokens(TArray<uint8>& TokenStream, TArray<FGameplayTag>& TagDictionary) const
+{
+	if (RootExpression)
+	{
+		TokenStream.Add(1);		// true if has a root expression
+		RootExpression->EmitTokens(TokenStream, TagDictionary);
+	}
+	else
+	{
+		TokenStream.Add(0);		// false if no root expression
+	}
+}
+
+void UEditableGameplayTagQueryExpression::EmitTagTokens(FGameplayTagContainer const& TagsToEmit, TArray<uint8>& TokenStream, TArray<FGameplayTag>& TagDictionary) const
+{
+	uint8 const NumTags = (uint8)TagsToEmit.Num();
+	TokenStream.Add(NumTags);
+
+	for (auto T : TagsToEmit)
+	{
+		int32 TagIdx = TagDictionary.AddUnique(T);
+		check(TagIdx <= 255);
+		TokenStream.Add((uint8)TagIdx);
+	}
+}
+
+void UEditableGameplayTagQueryExpression::EmitExprListTokens(TArray<UEditableGameplayTagQueryExpression*> const& ExprList, TArray<uint8>& TokenStream, TArray<FGameplayTag>& TagDictionary) const
+{
+	uint8 const NumExprs = (uint8)ExprList.Num();
+	TokenStream.Add(NumExprs);
+	int32 const NumExprsStreamLocation = TokenStream.Num() - 1;		// so we can fill in the actual number of written exprs later
+
+	for (auto E : ExprList)
+	{
+		if (E)
+		{
+			E->EmitTokens(TokenStream, TagDictionary);
+		}
+		else
+		{
+			// null expression
+			TokenStream.Add(EGameplayTagQueryExprType::Undefined);
+		}
+	}
+}
+
+void UEditableGameplayTagQueryExpression_AnyTagsMatch::EmitTokens(TArray<uint8>& TokenStream, TArray<FGameplayTag>& TagDictionary) const 
+{
+	TokenStream.Add(EGameplayTagQueryExprType::AnyTagsMatch);
+	EmitTagTokens(Tags, TokenStream, TagDictionary);
+}
+
+void UEditableGameplayTagQueryExpression_AllTagsMatch::EmitTokens(TArray<uint8>& TokenStream, TArray<FGameplayTag>& TagDictionary) const
+{
+	TokenStream.Add(EGameplayTagQueryExprType::AllTagsMatch);
+	EmitTagTokens(Tags, TokenStream, TagDictionary);
+}
+
+void UEditableGameplayTagQueryExpression_NoTagsMatch::EmitTokens(TArray<uint8>& TokenStream, TArray<FGameplayTag>& TagDictionary) const
+{
+	TokenStream.Add(EGameplayTagQueryExprType::NoTagsMatch);
+	EmitTagTokens(Tags, TokenStream, TagDictionary);
+}
+
+void UEditableGameplayTagQueryExpression_AnyExprMatch::EmitTokens(TArray<uint8>& TokenStream, TArray<FGameplayTag>& TagDictionary) const
+{
+	TokenStream.Add(EGameplayTagQueryExprType::AnyExprMatch);
+	EmitExprListTokens(Expressions, TokenStream, TagDictionary);
+}
+
+void UEditableGameplayTagQueryExpression_AllExprMatch::EmitTokens(TArray<uint8>& TokenStream, TArray<FGameplayTag>& TagDictionary) const
+{
+	TokenStream.Add(EGameplayTagQueryExprType::AllExprMatch);
+	EmitExprListTokens(Expressions, TokenStream, TagDictionary);
+}
+
+void UEditableGameplayTagQueryExpression_NoExprMatch::EmitTokens(TArray<uint8>& TokenStream, TArray<FGameplayTag>& TagDictionary) const
+{
+	TokenStream.Add(EGameplayTagQueryExprType::NoExprMatch);
+	EmitExprListTokens(Expressions, TokenStream, TagDictionary);
+}
+#endif	// WITH_EDITOR
+
+
+
+void FGameplayTagQueryExpression::EmitTokens(TArray<uint8>& TokenStream, TArray<FGameplayTag>& TagDictionary) const
+{
+	// emit exprtype
+	TokenStream.Add(ExprType);
+
+	// emit exprdata
+	switch (ExprType)
+	{
+	case EGameplayTagQueryExprType::AnyTagsMatch:
+	case EGameplayTagQueryExprType::AllTagsMatch:
+	case EGameplayTagQueryExprType::NoTagsMatch:
+	{
+		// emit tagset
+		uint8 NumTags = (uint8)TagSet.Num();
+		TokenStream.Add(NumTags);
+
+		for (auto Tag : TagSet)
+		{
+			int32 TagIdx = TagDictionary.AddUnique(Tag);
+			check(TagIdx <= 255);
+			TokenStream.Add((uint8)TagIdx);
+		}
+	}
+	break;
+
+	case EGameplayTagQueryExprType::AnyExprMatch:
+	case EGameplayTagQueryExprType::AllExprMatch:
+	case EGameplayTagQueryExprType::NoExprMatch:
+	{
+		// emit tagset
+		uint8 NumExprs = (uint8)ExprSet.Num();
+		TokenStream.Add(NumExprs);
+
+		for (auto& E : ExprSet)
+		{
+			E.EmitTokens(TokenStream, TagDictionary);
+		}
+	}
+	break;
+	default:
+		break;
+	}
+}
+
