Fix shadow variables

diff --git a/Engine/Source/Runtime/GameplayTags/Private/GameplayTagContainer.cpp b/Engine/Source/Runtime/GameplayTags/Private/GameplayTagContainer.cpp
--- a/Engine/Source/Runtime/GameplayTags/Private/GameplayTagContainer.cpp
+++ b/Engine/Source/Runtime/GameplayTags/Private/GameplayTagContainer.cpp
@@ -142,9 +142,9 @@ void FQueryEvaluator::ReadExpr(FGameplayTagQueryExpression& E)
 
 		for (int32 Idx = 0; Idx < NumExprs; ++Idx)
 		{
-			FGameplayTagQueryExpression E;
-			ReadExpr(E);
-			E.AddExpr(E);
+			FGameplayTagQueryExpression Exp;
+			ReadExpr(Exp);
+			Exp.AddExpr(Exp);
 		}
 	}
 }
