- Added description display to GameplayTagQuery customization. Description is by default an autogenerated string form of the query, but can be overridden by users for more descriptive comments ("Match enemies on fire")
- Added "clear all" button to the property customization as well

diff --git a/Engine/Source/Editor/GameplayTagsEditor/Private/GameplayTagQueryCustomization.cpp b/Engine/Source/Editor/GameplayTagsEditor/Private/GameplayTagQueryCustomization.cpp
--- a/Engine/Source/Editor/GameplayTagsEditor/Private/GameplayTagQueryCustomization.cpp
+++ b/Engine/Source/Editor/GameplayTagsEditor/Private/GameplayTagQueryCustomization.cpp
@@ -37,12 +37,38 @@ void FGameplayTagQueryCustomization::CustomizeHeader(TSharedRef<class IPropertyH
 				.Text(this, &FGameplayTagQueryCustomization::GetEditButtonText)
 				.OnClicked(this, &FGameplayTagQueryCustomization::OnEditButtonClicked)
 			]
+			+SVerticalBox::Slot()
+			.AutoHeight()
+			[
+				SNew(SButton)
+				.IsEnabled(!StructPropertyHandle->GetProperty()->HasAnyPropertyFlags(CPF_EditConst))
+				.Text(LOCTEXT("GameplayTagQueryCustomization_Clear", "Clear All"))
+				.OnClicked(this, &FGameplayTagQueryCustomization::OnClearAllButtonClicked)
+				.Visibility(this, &FGameplayTagQueryCustomization::GetClearAllVisibility)
+			]
+		]
+		+SHorizontalBox::Slot()
+		.AutoWidth()
+		[
+			SNew(SBorder)
+			.Padding(4.0f)
+			.Visibility(this, &FGameplayTagQueryCustomization::GetQueryDescVisibility)
+			[
+				SNew(STextBlock)
+				.Text(this, &FGameplayTagQueryCustomization::GetQueryDescText)
+				.AutoWrapText(true)
+			]
 		]
 	];
 
 	GEditor->RegisterForUndo(this);
 }
 
+FText FGameplayTagQueryCustomization::GetQueryDescText() const
+{
+	return FText::FromString(QueryDescription);
+}
+
 FText FGameplayTagQueryCustomization::GetEditButtonText() const
 {
 	bool const bReadOnly = StructPropertyHandle->GetProperty()->HasAnyPropertyFlags(CPF_EditConst);
@@ -52,6 +78,61 @@ FText FGameplayTagQueryCustomization::GetEditButtonText() const
 		: LOCTEXT("GameplayTagQueryCustomization_Edit", "Edit...");
 }
 
+FReply FGameplayTagQueryCustomization::OnClearAllButtonClicked()
+{
+	for (auto& EQ : EditableQueries)
+	{
+		if (EQ.TagQuery)
+		{
+			EQ.TagQuery->Clear();
+		}
+	}
+
+	RefreshQueryDescription();
+
+	return FReply::Handled();
+}
+
+EVisibility FGameplayTagQueryCustomization::GetClearAllVisibility() const
+{
+	bool bAtLeastOneQueryIsNonEmpty = false;
+
+	for (auto& EQ : EditableQueries)
+	{
+		if (EQ.TagQuery && EQ.TagQuery->IsEmpty() == false)
+		{
+			bAtLeastOneQueryIsNonEmpty = true;
+			break;
+		}
+	}
+
+	return bAtLeastOneQueryIsNonEmpty ? EVisibility::Visible : EVisibility::Collapsed;
+}
+
+EVisibility FGameplayTagQueryCustomization::GetQueryDescVisibility() const
+{
+	return QueryDescription.IsEmpty() ? EVisibility::Collapsed : EVisibility::Visible;
+}
+
+void FGameplayTagQueryCustomization::RefreshQueryDescription()
+{
+	// Rebuild Editable Containers as container references can become unsafe
+	BuildEditableQueryList();
+
+	// Clear the list
+	QueryDescription.Empty();
+
+	if (EditableQueries.Num() > 1)
+	{
+		QueryDescription = TEXT("Multiple Selected");
+	}
+	else if ( (EditableQueries.Num() == 1) && (EditableQueries[0].TagQuery != nullptr) )
+	{
+		QueryDescription = EditableQueries[0].TagQuery->GetDescription();
+	}
+}
+
+
 FReply FGameplayTagQueryCustomization::OnEditButtonClicked()
 {
 	if (GameplayTagQueryWidgetWindow.IsValid())
@@ -141,6 +222,8 @@ void FGameplayTagQueryCustomization::CloseWidgetWindow()
  	{
  		GameplayTagQueryWidgetWindow->RequestDestroyWindow();
 		GameplayTagQueryWidgetWindow = nullptr;
+
+		RefreshQueryDescription();
  	}
 }
 
diff --git a/Engine/Source/Editor/GameplayTagsEditor/Private/GameplayTagQueryCustomization.h b/Engine/Source/Editor/GameplayTagsEditor/Private/GameplayTagQueryCustomization.h
--- a/Engine/Source/Editor/GameplayTagsEditor/Private/GameplayTagQueryCustomization.h
+++ b/Engine/Source/Editor/GameplayTagsEditor/Private/GameplayTagQueryCustomization.h
@@ -29,6 +29,19 @@ private:
 	/** Overridden to do nothing */
 	FText GetEditButtonText() const;
 
+	/** Called when the "clear all" button is clicked */
+	FReply OnClearAllButtonClicked();
+
+	/** Returns the visibility of the "clear all" button (collapsed when there are no tags) */
+	EVisibility GetClearAllVisibility() const;
+
+	/** Returns the visibility of the tags list box (collapsed when there are no tags) */
+	EVisibility GetQueryDescVisibility() const;
+
+	void RefreshQueryDescription();
+
+	FText GetQueryDescText() const;
+
 	void CloseWidgetWindow();
 
 	/** Build List of Editable Queries */
@@ -42,5 +55,7 @@ private:
 
 	/** The Window for the GameplayTagWidget */
 	TSharedPtr<SWindow> GameplayTagQueryWidgetWindow;
+
+	FString QueryDescription;
 };
 
diff --git a/Engine/Source/Runtime/GameplayTags/Classes/GameplayTagContainer.h b/Engine/Source/Runtime/GameplayTags/Classes/GameplayTagContainer.h
--- a/Engine/Source/Runtime/GameplayTags/Classes/GameplayTagContainer.h
+++ b/Engine/Source/Runtime/GameplayTags/Classes/GameplayTagContainer.h
@@ -410,6 +410,10 @@ private:
 	UPROPERTY()
 	TArray<uint8> QueryTokenStream;
 
+	/** Human-readable string describing the query */
+	UPROPERTY()
+	FString Description;
+
 	/** Returns a gameplay tag from the tag dictionary */
 	FGameplayTag GetTagFromIndex(int32 TagIdx) const
 	{
@@ -417,15 +421,25 @@ private:
 		return TagDictionary[TagIdx];
 	}
 
+
 public:
 	/** Returns true if the given tags match this query, or false otherwise. */
 	bool Matches(FGameplayTagContainer const& Tags) const;
 
+	/** Returns true if this query is empty, false otherwise. */
+	bool IsEmpty() const;
+
+	/** Resets this query to its default empty state. */
+	void Clear();
+
 	/** Creates this query with the given root expression. */
 	void BuildQuery(struct FGameplayTagQueryExpression& RootQueryExpr);
 
 	/** Builds a FGameplayTagQueryExpression from this query. */
 	void GetQueryExpr(struct FGameplayTagQueryExpression& OutExpr) const;
+
+	/** Returns description string. */
+	FString GetDescription() const { return Description; };
 	
 #if WITH_EDITOR
 	/** Creates this query based on the given EditableQuery object */
@@ -536,13 +550,18 @@ class GAMEPLAYTAGS_API UEditableGameplayTagQuery : public UObject
 	GENERATED_BODY()
 
 public:
+
+	/** User-supplied description, shown in property details. Auto-generated description is shown if not supplied. */
+	UPROPERTY(EditDefaultsOnly, Category = Query)
+	FString Description;
+
 	/** The base expression of this query. */
 	UPROPERTY(EditDefaultsOnly, Instanced, Category = Query)
 	class UEditableGameplayTagQueryExpression* RootExpression;
 
 #if WITH_EDITOR
 	/** Converts this editor query construct into the runtime-usable token stream. */
-	void EmitTokens(TArray<uint8>& TokenStream, TArray<FGameplayTag>& TagDictionary) const;
+	void EmitTokens(TArray<uint8>& TokenStream, TArray<FGameplayTag>& TagDictionary, FString* DebugString=nullptr) const;
 #endif  // WITH_EDITOR
 };
 
@@ -554,11 +573,11 @@ class GAMEPLAYTAGS_API UEditableGameplayTagQueryExpression : public UObject
 #if WITH_EDITOR
 public:
 	/** Converts this editor query construct into the runtime-usable token stream. */
-	virtual void EmitTokens(TArray<uint8>& TokenStream, TArray<FGameplayTag>& TagDictionary) const {};
+	virtual void EmitTokens(TArray<uint8>& TokenStream, TArray<FGameplayTag>& TagDictionary, FString* DebugString=nullptr) const {};
 
 protected:
- 	void EmitTagTokens(FGameplayTagContainer const& TagsToEmit, TArray<uint8>& TokenStream, TArray<FGameplayTag>& TagDictionary) const;
- 	void EmitExprListTokens(TArray<UEditableGameplayTagQueryExpression*> const& ExprList, TArray<uint8>& TokenStream, TArray<FGameplayTag>& TagDictionary) const;
+	void EmitTagTokens(FGameplayTagContainer const& TagsToEmit, TArray<uint8>& TokenStream, TArray<FGameplayTag>& TagDictionary, FString* DebugString) const;
+	void EmitExprListTokens(TArray<UEditableGameplayTagQueryExpression*> const& ExprList, TArray<uint8>& TokenStream, TArray<FGameplayTag>& TagDictionary, FString* DebugString) const;
 #endif  // WITH_EDITOR
 };
 
@@ -571,7 +590,7 @@ public:
 	FGameplayTagContainer Tags;
 
 #if WITH_EDITOR
-	virtual void EmitTokens(TArray<uint8>& TokenStream, TArray<FGameplayTag>& TagDictionary) const override;
+	virtual void EmitTokens(TArray<uint8>& TokenStream, TArray<FGameplayTag>& TagDictionary, FString* DebugString = nullptr) const override;
 #endif  // WITH_EDITOR
 };
 
@@ -584,7 +603,7 @@ public:
 	FGameplayTagContainer Tags;
 
 #if WITH_EDITOR
-	virtual void EmitTokens(TArray<uint8>& TokenStream, TArray<FGameplayTag>& TagDictionary) const override;
+	virtual void EmitTokens(TArray<uint8>& TokenStream, TArray<FGameplayTag>& TagDictionary, FString* DebugString = nullptr) const override;
 #endif  // WITH_EDITOR
 };
 
@@ -597,7 +616,7 @@ public:
 	FGameplayTagContainer Tags;
 
 #if WITH_EDITOR
-	virtual void EmitTokens(TArray<uint8>& TokenStream, TArray<FGameplayTag>& TagDictionary) const override;
+	virtual void EmitTokens(TArray<uint8>& TokenStream, TArray<FGameplayTag>& TagDictionary, FString* DebugString = nullptr) const override;
 #endif  // WITH_EDITOR
 };
 
@@ -610,7 +629,7 @@ public:
 	TArray<UEditableGameplayTagQueryExpression*> Expressions;
 
 #if WITH_EDITOR
-	virtual void EmitTokens(TArray<uint8>& TokenStream, TArray<FGameplayTag>& TagDictionary) const override;
+	virtual void EmitTokens(TArray<uint8>& TokenStream, TArray<FGameplayTag>& TagDictionary, FString* DebugString = nullptr) const override;
 #endif  // WITH_EDITOR
 };
 
@@ -623,7 +642,7 @@ public:
 	TArray<UEditableGameplayTagQueryExpression*> Expressions;
 
 #if WITH_EDITOR
-	virtual void EmitTokens(TArray<uint8>& TokenStream, TArray<FGameplayTag>& TagDictionary) const override;
+	virtual void EmitTokens(TArray<uint8>& TokenStream, TArray<FGameplayTag>& TagDictionary, FString* DebugString = nullptr) const override;
 #endif  // WITH_EDITOR
 };
 
@@ -636,7 +655,7 @@ public:
 	TArray<UEditableGameplayTagQueryExpression*> Expressions;
 
 #if WITH_EDITOR
-	virtual void EmitTokens(TArray<uint8>& TokenStream, TArray<FGameplayTag>& TagDictionary) const override;
+	virtual void EmitTokens(TArray<uint8>& TokenStream, TArray<FGameplayTag>& TagDictionary, FString* DebugString = nullptr) const override;
 #endif  // WITH_EDITOR
 };
 
diff --git a/Engine/Source/Runtime/GameplayTags/Private/GameplayTagContainer.cpp b/Engine/Source/Runtime/GameplayTags/Private/GameplayTagContainer.cpp
--- a/Engine/Source/Runtime/GameplayTags/Private/GameplayTagContainer.cpp
+++ b/Engine/Source/Runtime/GameplayTags/Private/GameplayTagContainer.cpp
@@ -823,6 +823,16 @@ bool FGameplayTagQuery::Matches(FGameplayTagContainer const& Tags) const
 	return QE.Eval(Tags);
 }
 
+bool FGameplayTagQuery::IsEmpty() const
+{
+	return (QueryTokenStream.Num() == 0);
+}
+
+void FGameplayTagQuery::Clear()
+{
+	*this = FGameplayTagQuery::EmptyQuery;
+}
+
 void FGameplayTagQuery::GetQueryExpr(FGameplayTagQueryExpression& OutExpr) const
 {
 	// build the FExpr tree from the token stream and return it
@@ -1011,92 +1021,190 @@ void FGameplayTagQuery::BuildFromEditableQuery(UEditableGameplayTagQuery& Editab
 	QueryTokenStream.Reset();
 	TagDictionary.Reset();
 
+	Description = EditableQuery.Description;
+	FString* const StringRepresentation = Description.IsEmpty() ? &Description : nullptr;
+
 	// add stream version first
 	QueryTokenStream.Add(EGameplayTagQueryStreamVersion::LatestVersion);
-
-	EditableQuery.EmitTokens(QueryTokenStream, TagDictionary);
+	EditableQuery.EmitTokens(QueryTokenStream, TagDictionary, StringRepresentation);
 }
 
-void UEditableGameplayTagQuery::EmitTokens(TArray<uint8>& TokenStream, TArray<FGameplayTag>& TagDictionary) const
+void UEditableGameplayTagQuery::EmitTokens(TArray<uint8>& TokenStream, TArray<FGameplayTag>& TagDictionary, FString* DebugString) const
 {
+	if (DebugString)
+	{
+		// start with a fresh string
+		DebugString->Empty();
+	}
+
 	if (RootExpression)
 	{
 		TokenStream.Add(1);		// true if has a root expression
-		RootExpression->EmitTokens(TokenStream, TagDictionary);
+		RootExpression->EmitTokens(TokenStream, TagDictionary, DebugString);
 	}
 	else
 	{
 		TokenStream.Add(0);		// false if no root expression
+		DebugString->Append(TEXT("undefined"));
 	}
 }
 
-void UEditableGameplayTagQueryExpression::EmitTagTokens(FGameplayTagContainer const& TagsToEmit, TArray<uint8>& TokenStream, TArray<FGameplayTag>& TagDictionary) const
+void UEditableGameplayTagQueryExpression::EmitTagTokens(FGameplayTagContainer const& TagsToEmit, TArray<uint8>& TokenStream, TArray<FGameplayTag>& TagDictionary, FString* DebugString) const
 {
 	uint8 const NumTags = (uint8)TagsToEmit.Num();
 	TokenStream.Add(NumTags);
 
+	bool bFirstTag = true;
+
 	for (auto T : TagsToEmit)
 	{
 		int32 TagIdx = TagDictionary.AddUnique(T);
 		check(TagIdx <= 255);
 		TokenStream.Add((uint8)TagIdx);
+
+		if (DebugString)
+		{
+			if (bFirstTag == false)
+			{
+				DebugString->Append(TEXT(","));
+			}
+
+			DebugString->Append(TEXT(" "));
+			DebugString->Append(T.ToString());
+		}
+
+		bFirstTag = false;
 	}
 }
 
-void UEditableGameplayTagQueryExpression::EmitExprListTokens(TArray<UEditableGameplayTagQueryExpression*> const& ExprList, TArray<uint8>& TokenStream, TArray<FGameplayTag>& TagDictionary) const
+void UEditableGameplayTagQueryExpression::EmitExprListTokens(TArray<UEditableGameplayTagQueryExpression*> const& ExprList, TArray<uint8>& TokenStream, TArray<FGameplayTag>& TagDictionary, FString* DebugString) const
 {
 	uint8 const NumExprs = (uint8)ExprList.Num();
 	TokenStream.Add(NumExprs);
-	int32 const NumExprsStreamLocation = TokenStream.Num() - 1;		// so we can fill in the actual number of written exprs later
 
+	bool bFirstExpr = true;
+	
 	for (auto E : ExprList)
 	{
+		if (DebugString)
+		{
+			if (bFirstExpr == false)
+			{
+				DebugString->Append(TEXT(","));
+			}
+
+			DebugString->Append(TEXT(" "));
+		}
+
 		if (E)
 		{
-			E->EmitTokens(TokenStream, TagDictionary);
+			E->EmitTokens(TokenStream, TagDictionary, DebugString);
 		}
 		else
 		{
 			// null expression
 			TokenStream.Add(EGameplayTagQueryExprType::Undefined);
+			DebugString->Append(TEXT("undefined"));
 		}
+
+		bFirstExpr = false;
 	}
 }
 
-void UEditableGameplayTagQueryExpression_AnyTagsMatch::EmitTokens(TArray<uint8>& TokenStream, TArray<FGameplayTag>& TagDictionary) const 
+void UEditableGameplayTagQueryExpression_AnyTagsMatch::EmitTokens(TArray<uint8>& TokenStream, TArray<FGameplayTag>& TagDictionary, FString* DebugString) const
 {
 	TokenStream.Add(EGameplayTagQueryExprType::AnyTagsMatch);
-	EmitTagTokens(Tags, TokenStream, TagDictionary);
+
+	if (DebugString)
+	{
+		DebugString->Append(TEXT(" ANY("));
+	}
+
+	EmitTagTokens(Tags, TokenStream, TagDictionary, DebugString);
+
+	if (DebugString)
+	{
+		DebugString->Append(TEXT(" )"));
+	}
 }
 
-void UEditableGameplayTagQueryExpression_AllTagsMatch::EmitTokens(TArray<uint8>& TokenStream, TArray<FGameplayTag>& TagDictionary) const
+void UEditableGameplayTagQueryExpression_AllTagsMatch::EmitTokens(TArray<uint8>& TokenStream, TArray<FGameplayTag>& TagDictionary, FString* DebugString) const
 {
 	TokenStream.Add(EGameplayTagQueryExprType::AllTagsMatch);
-	EmitTagTokens(Tags, TokenStream, TagDictionary);
+
+	if (DebugString)
+	{
+		DebugString->Append(TEXT(" ALL("));
+	}
+	
+	EmitTagTokens(Tags, TokenStream, TagDictionary, DebugString);
+
+	if (DebugString)
+	{
+		DebugString->Append(TEXT(" )"));
+	}
 }
 
-void UEditableGameplayTagQueryExpression_NoTagsMatch::EmitTokens(TArray<uint8>& TokenStream, TArray<FGameplayTag>& TagDictionary) const
+void UEditableGameplayTagQueryExpression_NoTagsMatch::EmitTokens(TArray<uint8>& TokenStream, TArray<FGameplayTag>& TagDictionary, FString* DebugString) const
 {
 	TokenStream.Add(EGameplayTagQueryExprType::NoTagsMatch);
-	EmitTagTokens(Tags, TokenStream, TagDictionary);
+	EmitTagTokens(Tags, TokenStream, TagDictionary, DebugString);
+
+	if (DebugString)
+	{
+		DebugString->Append(TEXT(" )"));
+	}
 }
 
-void UEditableGameplayTagQueryExpression_AnyExprMatch::EmitTokens(TArray<uint8>& TokenStream, TArray<FGameplayTag>& TagDictionary) const
+void UEditableGameplayTagQueryExpression_AnyExprMatch::EmitTokens(TArray<uint8>& TokenStream, TArray<FGameplayTag>& TagDictionary, FString* DebugString) const
 {
 	TokenStream.Add(EGameplayTagQueryExprType::AnyExprMatch);
-	EmitExprListTokens(Expressions, TokenStream, TagDictionary);
+
+	if (DebugString)
+	{
+		DebugString->Append(TEXT(" ANY("));
+	}
+
+	EmitExprListTokens(Expressions, TokenStream, TagDictionary, DebugString);
+
+	if (DebugString)
+	{
+		DebugString->Append(TEXT(" )"));
+	}
 }
 
-void UEditableGameplayTagQueryExpression_AllExprMatch::EmitTokens(TArray<uint8>& TokenStream, TArray<FGameplayTag>& TagDictionary) const
+void UEditableGameplayTagQueryExpression_AllExprMatch::EmitTokens(TArray<uint8>& TokenStream, TArray<FGameplayTag>& TagDictionary, FString* DebugString) const
 {
 	TokenStream.Add(EGameplayTagQueryExprType::AllExprMatch);
-	EmitExprListTokens(Expressions, TokenStream, TagDictionary);
+
+	if (DebugString)
+	{
+		DebugString->Append(TEXT(" ALL("));
+	}
+
+	EmitExprListTokens(Expressions, TokenStream, TagDictionary, DebugString);
+
+	if (DebugString)
+	{
+		DebugString->Append(TEXT(" )"));
+	}
 }
 
-void UEditableGameplayTagQueryExpression_NoExprMatch::EmitTokens(TArray<uint8>& TokenStream, TArray<FGameplayTag>& TagDictionary) const
+void UEditableGameplayTagQueryExpression_NoExprMatch::EmitTokens(TArray<uint8>& TokenStream, TArray<FGameplayTag>& TagDictionary, FString* DebugString) const
 {
 	TokenStream.Add(EGameplayTagQueryExprType::NoExprMatch);
-	EmitExprListTokens(Expressions, TokenStream, TagDictionary);
+
+	if (DebugString)
+	{
+		DebugString->Append(TEXT(" NONE("));
+	}
+
+	EmitExprListTokens(Expressions, TokenStream, TagDictionary, DebugString);
+
+	if (DebugString)
+	{
+		DebugString->Append(TEXT(" )"));
+	}
 }
 #endif	// WITH_EDITOR
 
