move INI GameplayTags into their own config file (now GameplayTags.ini instead of DefaultEngine.ini). Upgrade should be automatic.

--------
Integrated using branch UE4-To-UE4-Orion (reversed) of change#2595451 by [REDACTED_PERSON] on 2015/06/22 11:08:27.

diff --git a/Engine/Source/Runtime/GameplayTags/Classes/GameplayTagsSettings.h b/Engine/Source/Runtime/GameplayTags/Classes/GameplayTagsSettings.h
--- a/Engine/Source/Runtime/GameplayTags/Classes/GameplayTagsSettings.h
+++ b/Engine/Source/Runtime/GameplayTags/Classes/GameplayTagsSettings.h
@@ -19,7 +19,7 @@
  *	the user types (e.g, autocomplete 'Damage.Physical' as the user is adding a 'Damage.Physical.Slash' tag).
  *	
  */
-UCLASS(config=Engine, defaultconfig, notplaceable)
+UCLASS(config=GameplayTags, defaultconfig, notplaceable)
 class GAMEPLAYTAGS_API UGameplayTagsSettings : public UObject
 {
 	GENERATED_UCLASS_BODY()
diff --git a/Engine/Source/Runtime/GameplayTags/Private/GameplayTagsManager.cpp b/Engine/Source/Runtime/GameplayTags/Private/GameplayTagsManager.cpp
--- a/Engine/Source/Runtime/GameplayTags/Private/GameplayTagsManager.cpp
+++ b/Engine/Source/Runtime/GameplayTags/Private/GameplayTagsManager.cpp
@@ -137,6 +137,19 @@ void UGameplayTagsManager::ConstructGameplayTagTree()
 			FString PerfMessage = FString::Printf(TEXT("UGameplayTagsManager::ConstructGameplayTagTree: ImportINI"));
 			SCOPE_LOG_TIME_IN_SECONDS(*PerfMessage, nullptr)
 #endif
+
+			// Update path: Check for old tags in DefaultEngine.ini (we'll push them to the UGameplayTagSettings class).
+			TArray<FString> EngineConfigTags;
+			GConfig->GetArray(TEXT("/Script/GameplayTags.GameplayTagsSettings"), TEXT("GameplayTags"), EngineConfigTags, GEngineIni);
+			if (EngineConfigTags.Num() > 0)
+			{
+				UGameplayTagsSettings* MutableDefault = GetMutableDefault<UGameplayTagsSettings>();
+				if (MutableDefault->GameplayTags.Num() == 0)
+				{
+					MutableDefault->GameplayTags.Append(EngineConfigTags);
+				}
+			}
+
 			// Load any GameplayTagSettings from config (their default object)
 			for (TObjectIterator<UClass> ClassIt; ClassIt; ++ClassIt)
 			{
@@ -145,15 +158,13 @@ void UGameplayTagsManager::ConstructGameplayTagTree()
 				{
 					continue;
 				}
-
+				
 				for (FString TagStr : Class->GetDefaultObject<UGameplayTagsSettings>()->GameplayTags)
 				{
 					FGameplayTagTableRow TableRow;
 					TableRow.Tag = TagStr;
 					AddTagTableRow(TableRow);
 				}
-
-		
 			}
 			GameplayRootTag->GetChildTagNodes().Sort(FCompareFGameplayTagNodeByTag());
 		}
