alphabetize gameplay tags imported from INI file

--------
Integrated using branch UE4-To-UE4-Orion (reversed) of change#2602290 by [REDACTED_PERSON] on 2015/06/26 11:12:23.

diff --git a/Engine/Source/Editor/GameplayAbilitiesEditor/Private/SGameplayCueEditor.cpp b/Engine/Source/Editor/GameplayAbilitiesEditor/Private/SGameplayCueEditor.cpp
--- a/Engine/Source/Editor/GameplayAbilitiesEditor/Private/SGameplayCueEditor.cpp
+++ b/Engine/Source/Editor/GameplayAbilitiesEditor/Private/SGameplayCueEditor.cpp
@@ -444,6 +444,7 @@ private:
 			}
 
 			Settings->GameplayTags.Add(str);
+			Settings->SortTags();
 			IGameplayTagsModule::Get().GetGameplayTagsManager().DestroyGameplayTagTree();
 
 			{
diff --git a/Engine/Source/Runtime/GameplayTags/Classes/GameplayTagsSettings.h b/Engine/Source/Runtime/GameplayTags/Classes/GameplayTagsSettings.h
--- a/Engine/Source/Runtime/GameplayTags/Classes/GameplayTagsSettings.h
+++ b/Engine/Source/Runtime/GameplayTags/Classes/GameplayTagsSettings.h
@@ -27,6 +27,9 @@ class GAMEPLAYTAGS_API UGameplayTagsSettings : public UObject
 	UPROPERTY(config, EditAnywhere, Category=GameplayTags)
 	TArray<FString>		GameplayTags;
 
+	/** Sorts tags alphabetically */
+	void SortTags();
+
 #if WITH_EDITOR
 	virtual void PostEditChangeProperty(FPropertyChangedEvent& PropertyChangedEvent) override;
 #endif
diff --git a/Engine/Source/Runtime/GameplayTags/Private/GameplayTagsManager.cpp b/Engine/Source/Runtime/GameplayTags/Private/GameplayTagsManager.cpp
--- a/Engine/Source/Runtime/GameplayTags/Private/GameplayTagsManager.cpp
+++ b/Engine/Source/Runtime/GameplayTags/Private/GameplayTagsManager.cpp
@@ -158,6 +158,10 @@ void UGameplayTagsManager::ConstructGameplayTagTree()
 				{
 					continue;
 				}
+
+#if WITH_EDITOR
+				Class->GetDefaultObject<UGameplayTagsSettings>()->SortTags();
+#endif
 				
 				for (FString TagStr : Class->GetDefaultObject<UGameplayTagsSettings>()->GameplayTags)
 				{
diff --git a/Engine/Source/Runtime/GameplayTags/Private/GameplayTagsSettings.cpp b/Engine/Source/Runtime/GameplayTags/Private/GameplayTagsSettings.cpp
--- a/Engine/Source/Runtime/GameplayTags/Private/GameplayTagsSettings.cpp
+++ b/Engine/Source/Runtime/GameplayTags/Private/GameplayTagsSettings.cpp
@@ -24,3 +24,8 @@ void UGameplayTagsSettings::PostEditChangeProperty(FPropertyChangedEvent& Proper
 	}
 }
 #endif
+
+void UGameplayTagsSettings::SortTags()
+{
+	GameplayTags.Sort();
+}
\ No newline at end of file
