Added UBlueprintGameplayTagLibrary::DoesContainerMatchTagQuery()

diff --git a/Engine/Source/Runtime/GameplayTags/Classes/BlueprintGameplayTagLibrary.h b/Engine/Source/Runtime/GameplayTags/Classes/BlueprintGameplayTagLibrary.h
--- a/Engine/Source/Runtime/GameplayTags/Classes/BlueprintGameplayTagLibrary.h
+++ b/Engine/Source/Runtime/GameplayTags/Classes/BlueprintGameplayTagLibrary.h
@@ -75,6 +75,17 @@ class UBlueprintGameplayTagLibrary : public UBlueprintFunctionLibrary
 	UFUNCTION(BlueprintPure, Category="GameplayTags|Tag Container")
 	static bool DoesContainerMatchAllTagsInContainer(const FGameplayTagContainer& TagContainer, const FGameplayTagContainer& OtherContainer, bool bCountEmptyAsMatch);
 
+	/**
+	 * Check if the specified tag container matches the given Tag Query.
+	 * 
+	 * @param TagContainer			Container to check if it matches all of the tags in the other container
+	 * @param TagQuery				Query to match against
+	 * 
+	 * @return True if the container matches the query, false otherwise.
+	 */
+	UFUNCTION(BlueprintPure, Category = "GameplayTags|Tag Container")
+	static bool DoesContainerMatchTagQuery(const FGameplayTagContainer& TagContainer, const FGameplayTagQuery& TagQuery);
+
 	/**
 	* Takes 2 Gameplay Tag Containers Appends the In Container to the In/Out Container
 	*
diff --git a/Engine/Source/Runtime/GameplayTags/Private/BlueprintGameplayTagLibrary.cpp b/Engine/Source/Runtime/GameplayTags/Private/BlueprintGameplayTagLibrary.cpp
--- a/Engine/Source/Runtime/GameplayTags/Private/BlueprintGameplayTagLibrary.cpp
+++ b/Engine/Source/Runtime/GameplayTags/Private/BlueprintGameplayTagLibrary.cpp
@@ -32,6 +32,11 @@ bool UBlueprintGameplayTagLibrary::DoesContainerMatchAllTagsInContainer(const FG
 	return TagContainer.MatchesAll(OtherContainer, bCountEmptyAsMatch);
 }
 
+bool UBlueprintGameplayTagLibrary::DoesContainerMatchTagQuery(const FGameplayTagContainer& TagContainer, const FGameplayTagQuery& TagQuery)
+{
+	return TagQuery.Matches(TagContainer);
+}
+
 FGameplayTag UBlueprintGameplayTagLibrary::MakeLiteralGameplayTag(FGameplayTag Value)
 {
 	return Value;
