Added MakeGameplayTagQuery blueprint node, since the auto-generated one doesn't work for this struct.

diff --git a/Engine/Source/Runtime/GameplayTags/Classes/BlueprintGameplayTagLibrary.h b/Engine/Source/Runtime/GameplayTags/Classes/BlueprintGameplayTagLibrary.h
--- a/Engine/Source/Runtime/GameplayTags/Classes/BlueprintGameplayTagLibrary.h
+++ b/Engine/Source/Runtime/GameplayTags/Classes/BlueprintGameplayTagLibrary.h
@@ -105,6 +105,17 @@ class UBlueprintGameplayTagLibrary : public UBlueprintFunctionLibrary
 	UFUNCTION(BlueprintPure, Category = "GameplayTags")
 	static FGameplayTag MakeLiteralGameplayTag(FGameplayTag Value);
 
+
+	/**
+	 * Creates a literal FGameplayTagQuery
+	 *
+	 * @param	Value	value to set the FGameplayTagQuery to
+	 *
+	 * @return	The literal FGameplayTagQuery
+	 */
+	UFUNCTION(BlueprintPure, Category = "GameplayTags")
+	static FGameplayTagQuery MakeGameplayTagQuery(FGameplayTagQuery TagQuery);
+
 	/**
 	 * Check Gameplay tags in the interface has all of the specified tags in the tag container (expands to include parents of asset tags)
 	 *
diff --git a/Engine/Source/Runtime/GameplayTags/Classes/GameplayTagContainer.h b/Engine/Source/Runtime/GameplayTags/Classes/GameplayTagContainer.h
--- a/Engine/Source/Runtime/GameplayTags/Classes/GameplayTagContainer.h
+++ b/Engine/Source/Runtime/GameplayTags/Classes/GameplayTagContainer.h
@@ -396,7 +396,7 @@ namespace EGameplayTagQueryStreamVersion
  * Queries are internally represented as a byte stream that is memory-efficient and can be evaluated quickly at runtime.
  * Note: these have an extensive details and graph pin customization for editing, so there is no need to expose the internals to Blueprints.
  */
-USTRUCT(BlueprintType)
+USTRUCT(BlueprintType, meta=(HasNativeMake="GameplayTags.UBlueprintGameplayTagLibrary.MakeGameplayTagQuery"))
 struct GAMEPLAYTAGS_API FGameplayTagQuery
 {
 	GENERATED_BODY();
diff --git a/Engine/Source/Runtime/GameplayTags/Private/BlueprintGameplayTagLibrary.cpp b/Engine/Source/Runtime/GameplayTags/Private/BlueprintGameplayTagLibrary.cpp
--- a/Engine/Source/Runtime/GameplayTags/Private/BlueprintGameplayTagLibrary.cpp
+++ b/Engine/Source/Runtime/GameplayTags/Private/BlueprintGameplayTagLibrary.cpp
@@ -42,6 +42,11 @@ FGameplayTag UBlueprintGameplayTagLibrary::MakeLiteralGameplayTag(FGameplayTag V
 	return Value;
 }
 
+FGameplayTagQuery UBlueprintGameplayTagLibrary::MakeGameplayTagQuery(FGameplayTagQuery TagQuery)
+{
+	return TagQuery;
+}
+
 bool UBlueprintGameplayTagLibrary::HasAllMatchingGameplayTags(TScriptInterface<IGameplayTagAssetInterface> TagContainerInterface, const FGameplayTagContainer& OtherContainer, bool bCountEmptyAsMatch)
 {
 	if (TagContainerInterface.GetInterface() == NULL)
