Remove virtual functions from GameplayTag container. Added RemoveAllTagsKeepSlack helper functions.

--------
Integrated using branch UE4-To-UE4-Orion (reversed) of change#2667398 by [REDACTED_PERSON] on 2015/08/25 09:20:09.

diff --git a/Engine/Source/Runtime/GameplayTags/Classes/GameplayTagContainer.h b/Engine/Source/Runtime/GameplayTags/Classes/GameplayTagContainer.h
--- a/Engine/Source/Runtime/GameplayTags/Classes/GameplayTagContainer.h
+++ b/Engine/Source/Runtime/GameplayTags/Classes/GameplayTagContainer.h
@@ -186,14 +186,14 @@ struct GAMEPLAYTAGS_API FGameplayTagContainer
 	 * 
 	 * @return True if the tag is in the container, false if it is not
 	 */
-	virtual bool HasTag(FGameplayTag const& TagToCheck, TEnumAsByte<EGameplayTagMatchType::Type> TagMatchType, TEnumAsByte<EGameplayTagMatchType::Type> TagToCheckMatchType) const;
+	bool HasTag(FGameplayTag const& TagToCheck, TEnumAsByte<EGameplayTagMatchType::Type> TagMatchType, TEnumAsByte<EGameplayTagMatchType::Type> TagToCheckMatchType) const;
 
 	/** 
 	 * Adds all the tags from one container to this container 
 	 *
 	 * @param Other TagContainer that has the tags you want to add to this container 
 	 */
-	virtual void AppendTags(FGameplayTagContainer const& Other);
+	void AppendTags(FGameplayTagContainer const& Other);
 
 	/** 
 	 * Adds all the tags that match between the two specified containers to this container 
@@ -201,7 +201,7 @@ struct GAMEPLAYTAGS_API FGameplayTagContainer
 	 * @param OtherA TagContainer that has the matching tags you want to add to this container 
 	 * @param OtherB TagContainer used to check for matching tags
 	 */
-	virtual void AppendMatchingTags(FGameplayTagContainer const& OtherA, FGameplayTagContainer const& OtherB);
+	void AppendMatchingTags(FGameplayTagContainer const& OtherA, FGameplayTagContainer const& OtherB);
 
 	/**
 	 * Add the specified tag to the container
@@ -224,17 +224,22 @@ struct GAMEPLAYTAGS_API FGameplayTagContainer
 	 * 
 	 * @param TagToRemove	Tag to remove from the container
 	 */
-	virtual void RemoveTag(FGameplayTag TagToRemove);
+	void RemoveTag(FGameplayTag TagToRemove);
 
 	/**
 	* Removes all tags in TagsToRemove from this container
 	*
 	* @param TagsToRemove	Tags to remove from the container
 	*/
-	virtual void RemoveTags(FGameplayTagContainer TagsToRemove);
+	void RemoveTags(FGameplayTagContainer TagsToRemove);
 
 	/** Remove all tags from the container */
-	virtual void RemoveAllTags(int32 Slack=0);
+	void RemoveAllTags(int32 Slack=0);
+
+	void RemoveAllTagsKeepSlack()
+	{
+		RemoveAllTags(GameplayTags.Num());
+	}
 
 	/**
 	 * Serialize the tag container
