Add bool return parameter RemoveTag to specify whether a tag was removed

--------
Integrated using branch UE4-To-UE4-Orion (reversed) of change#2667999 by [REDACTED_PERSON] on 2015/08/25 14:07:08.

diff --git a/Engine/Source/Runtime/GameplayTags/Classes/GameplayTagContainer.h b/Engine/Source/Runtime/GameplayTags/Classes/GameplayTagContainer.h
--- a/Engine/Source/Runtime/GameplayTags/Classes/GameplayTagContainer.h
+++ b/Engine/Source/Runtime/GameplayTags/Classes/GameplayTagContainer.h
@@ -224,7 +224,7 @@ struct GAMEPLAYTAGS_API FGameplayTagContainer
 	 * 
 	 * @param TagToRemove	Tag to remove from the container
 	 */
-	void RemoveTag(FGameplayTag TagToRemove);
+	bool RemoveTag(FGameplayTag TagToRemove);
 
 	/**
 	* Removes all tags in TagsToRemove from this container
diff --git a/Engine/Source/Runtime/GameplayTags/Private/GameplayTagContainer.cpp b/Engine/Source/Runtime/GameplayTags/Private/GameplayTagContainer.cpp
--- a/Engine/Source/Runtime/GameplayTags/Private/GameplayTagContainer.cpp
+++ b/Engine/Source/Runtime/GameplayTags/Private/GameplayTagContainer.cpp
@@ -609,9 +609,9 @@ void FGameplayTagContainer::AddTagFast(const FGameplayTag& TagToAdd)
 	GameplayTags.Add(TagToAdd);
 }
 
-void FGameplayTagContainer::RemoveTag(FGameplayTag TagToRemove)
+bool FGameplayTagContainer::RemoveTag(FGameplayTag TagToRemove)
 {
-	GameplayTags.Remove(TagToRemove);
+	return GameplayTags.RemoveSingle(TagToRemove) > 0;
 }
 
 void FGameplayTagContainer::RemoveTags(FGameplayTagContainer TagsToRemove)
