Merged CL 2669254 form UE4-Orion/... to UE4/...

Added GGameplayTagsManager to avoid a Module::Get() and a virtual function call every time we call HasTag() or other tag-related functions.

diff --git a/Engine/Source/Editor/GameplayAbilitiesEditor/Private/GameplayAbilitiesEditorModule.cpp b/Engine/Source/Editor/GameplayAbilitiesEditor/Private/GameplayAbilitiesEditorModule.cpp
--- a/Engine/Source/Editor/GameplayAbilitiesEditor/Private/GameplayAbilitiesEditorModule.cpp
+++ b/Engine/Source/Editor/GameplayAbilitiesEditor/Private/GameplayAbilitiesEditorModule.cpp
@@ -113,7 +113,7 @@ void FGameplayAbilitiesEditorModule::StartupModule()
 	FEdGraphUtilities::RegisterVisualNodeFactory(GameplayAbilitiesGraphPanelNodeFactory);
 
 	// Listen for changes to the gameplay tag tree so we can refresh blueprint actions for the GameplayCueEvent node
-	UGameplayTagsManager& GameplayTagsManager = IGameplayTagsModule::Get().GetGameplayTagsManager();
+	UGameplayTagsManager& GameplayTagsManager = IGameplayTagsModule::GetGameplayTagsManager();
 	GameplayTagTreeChangedDelegateHandle = GameplayTagsManager.OnGameplayTagTreeChanged().AddStatic(&FGameplayAbilitiesEditorModule::GameplayTagTreeChanged);
 
 	// GameplayCue editor
@@ -190,7 +190,7 @@ void FGameplayAbilitiesEditorModule::ShutdownModule()
 
 	if ( UObjectInitialized() && IGameplayTagsModule::IsAvailable() )
 	{
-		UGameplayTagsManager& GameplayTagsManager = IGameplayTagsModule::Get().GetGameplayTagsManager();
+		UGameplayTagsManager& GameplayTagsManager = IGameplayTagsModule::GetGameplayTagsManager();
 		GameplayTagsManager.OnGameplayTagTreeChanged().Remove(GameplayTagTreeChangedDelegateHandle);
 	}
 }
diff --git a/Engine/Source/Editor/GameplayTagsEditor/Private/SGameplayTagWidget.cpp b/Engine/Source/Editor/GameplayTagsEditor/Private/SGameplayTagWidget.cpp
--- a/Engine/Source/Editor/GameplayTagsEditor/Private/SGameplayTagWidget.cpp
+++ b/Engine/Source/Editor/GameplayTagsEditor/Private/SGameplayTagWidget.cpp
@@ -304,7 +304,7 @@ void SGameplayTagWidget::OnTagChecked(TSharedPtr<FGameplayTagNode> NodeChecked)
 {
 	FScopedTransaction Transaction( LOCTEXT("GameplayTagWidget_AddTags", "Add Gameplay Tags") );
 
-	UGameplayTagsManager& TagsManager = IGameplayTagsModule::Get().GetGameplayTagsManager();
+	UGameplayTagsManager& TagsManager = IGameplayTagsModule::GetGameplayTagsManager();
 
 	for (int32 ContainerIdx = 0; ContainerIdx < TagContainers.Num(); ++ContainerIdx)
 	{
@@ -348,7 +348,7 @@ void SGameplayTagWidget::OnTagUnchecked(TSharedPtr<FGameplayTagNode> NodeUncheck
 	FScopedTransaction Transaction( LOCTEXT("GameplayTagWidget_RemoveTags", "Remove Gameplay Tags"));
 	if (NodeUnchecked.IsValid())
 	{
-		UGameplayTagsManager& TagsManager = IGameplayTagsModule::Get().GetGameplayTagsManager();
+		UGameplayTagsManager& TagsManager = IGameplayTagsModule::GetGameplayTagsManager();
 
 		for (int32 ContainerIdx = 0; ContainerIdx < TagContainers.Num(); ++ContainerIdx)
 		{
@@ -398,7 +398,7 @@ void SGameplayTagWidget::OnTagUnchecked(TSharedPtr<FGameplayTagNode> NodeUncheck
 
 void SGameplayTagWidget::UncheckChildren(TSharedPtr<FGameplayTagNode> NodeUnchecked, FGameplayTagContainer& EditableContainer)
 {
-	UGameplayTagsManager& TagsManager = IGameplayTagsModule::Get().GetGameplayTagsManager();
+	UGameplayTagsManager& TagsManager = IGameplayTagsModule::GetGameplayTagsManager();
 
 	FGameplayTag Tag = TagsManager.RequestGameplayTag(NodeUnchecked->GetCompleteTag());
 	EditableContainer.RemoveTag(Tag);
@@ -417,7 +417,7 @@ ECheckBoxState SGameplayTagWidget::IsTagChecked(TSharedPtr<FGameplayTagNode> Nod
 
 	if (Node.IsValid())
 	{
-		UGameplayTagsManager& TagsManager = IGameplayTagsModule::Get().GetGameplayTagsManager();
+		UGameplayTagsManager& TagsManager = IGameplayTagsModule::GetGameplayTagsManager();
 
 		for (int32 ContainerIdx = 0; ContainerIdx < TagContainers.Num(); ++ContainerIdx)
 		{
@@ -522,7 +522,7 @@ void SGameplayTagWidget::VerifyAssetTagValidity()
 	// Create a set that is the library of all valid tags
 	TArray< TSharedPtr<FGameplayTagNode> > NodeStack;
 
-	UGameplayTagsManager& TagsManager = IGameplayTagsModule::Get().GetGameplayTagsManager();
+	UGameplayTagsManager& TagsManager = IGameplayTagsModule::GetGameplayTagsManager();
 	
 	TagsManager.GetFilteredGameplayRootTags(TEXT(""), NodeStack);
 
diff --git a/Engine/Source/Runtime/Engine/Private/UnrealEngine.cpp b/Engine/Source/Runtime/Engine/Private/UnrealEngine.cpp
--- a/Engine/Source/Runtime/Engine/Private/UnrealEngine.cpp
+++ b/Engine/Source/Runtime/Engine/Private/UnrealEngine.cpp
@@ -12743,4 +12743,6 @@ int32 UEngine::RenderStatSlateBatches(UWorld* World, FViewport* Viewport, FCanva
 }
 #endif
 
+ENGINE_API UGameplayTagsManager* GGameplayTagsManager = nullptr;
+
 #undef LOCTEXT_NAMESPACE
diff --git a/Engine/Source/Runtime/Engine/Public/EngineGlobals.h b/Engine/Source/Runtime/Engine/Public/EngineGlobals.h
--- a/Engine/Source/Runtime/Engine/Public/EngineGlobals.h
+++ b/Engine/Source/Runtime/Engine/Public/EngineGlobals.h
@@ -11,4 +11,7 @@ extern ENGINE_API class UEngine*			GEngine;
 extern ENGINE_API bool						GDisallowNetworkTravel;
 
 /** The GPU time taken to render the last frame. Same metric as FPlatformTime::Cycles(). */
-extern ENGINE_API uint32					GGPUFrameTime;
\ No newline at end of file
+extern ENGINE_API uint32					GGPUFrameTime;
+
+/** A reference to the ability system so we do not need to get the module every time we access the tag manager. */
+extern ENGINE_API class UGameplayTagsManager* GGameplayTagsManager;
\ No newline at end of file
diff --git a/Engine/Source/Runtime/GameplayTags/Private/GameplayTagContainer.cpp b/Engine/Source/Runtime/GameplayTags/Private/GameplayTagContainer.cpp
--- a/Engine/Source/Runtime/GameplayTags/Private/GameplayTagContainer.cpp
+++ b/Engine/Source/Runtime/GameplayTags/Private/GameplayTagContainer.cpp
@@ -455,7 +455,7 @@ bool FGameplayTagContainer::operator!=(FGameplayTagContainer const& Other) const
 
 bool FGameplayTagContainer::HasTag(FGameplayTag const& TagToCheck, TEnumAsByte<EGameplayTagMatchType::Type> TagMatchType, TEnumAsByte<EGameplayTagMatchType::Type> TagToCheckMatchType) const
 {
-	UGameplayTagsManager& TagManager = IGameplayTagsModule::Get().GetGameplayTagsManager();
+	UGameplayTagsManager& TagManager = IGameplayTagsModule::GetGameplayTagsManager();
 	for (TArray<FGameplayTag>::TConstIterator It(this->GameplayTags); It; ++It)
 	{
 		if (TagManager.GameplayTagsMatch(*It, TagMatchType, TagToCheck, TagToCheckMatchType) == true)
@@ -497,7 +497,7 @@ FGameplayTagContainer FGameplayTagContainer::GetGameplayTagParents() const
 FGameplayTagContainer FGameplayTagContainer::Filter(const FGameplayTagContainer& OtherContainer, TEnumAsByte<EGameplayTagMatchType::Type> TagMatchType, TEnumAsByte<EGameplayTagMatchType::Type> OtherTagMatchType) const
 {
 	FGameplayTagContainer ResultContainer;
-	UGameplayTagsManager& TagManager = IGameplayTagsModule::Get().GetGameplayTagsManager();
+	UGameplayTagsManager& TagManager = IGameplayTagsModule::GetGameplayTagsManager();
 
 	
 	for (TArray<FGameplayTag>::TConstIterator OtherIt(OtherContainer.GameplayTags); OtherIt; ++OtherIt)
@@ -516,7 +516,7 @@ FGameplayTagContainer FGameplayTagContainer::Filter(const FGameplayTagContainer&
 
 bool FGameplayTagContainer::DoesTagContainerMatch(const FGameplayTagContainer& OtherContainer, TEnumAsByte<EGameplayTagMatchType::Type> TagMatchType, TEnumAsByte<EGameplayTagMatchType::Type> OtherTagMatchType, EGameplayContainerMatchType ContainerMatchType) const
 {
-	UGameplayTagsManager& TagManager = IGameplayTagsModule::Get().GetGameplayTagsManager();
+	UGameplayTagsManager& TagManager = IGameplayTagsModule::GetGameplayTagsManager();
 
 	for (TArray<FGameplayTag>::TConstIterator OtherIt(OtherContainer.GameplayTags); OtherIt; ++OtherIt)
 	{
@@ -642,7 +642,7 @@ bool FGameplayTagContainer::Serialize(FArchive& Ar)
 	
 	if (Ar.IsLoading())
 	{
-		UGameplayTagsManager& TagManager = IGameplayTagsModule::Get().GetGameplayTagsManager();
+		UGameplayTagsManager& TagManager = IGameplayTagsModule::GetGameplayTagsManager();
 
 		// If loading old version, add old tags to the new gameplay tags array so they can be saved out with the new version
 		// This needs to happen 
@@ -785,7 +785,7 @@ FGameplayTag FGameplayTag::RequestDirectParent() const
 
 bool FGameplayTag::NetSerialize(FArchive& Ar, class UPackageMap* Map, bool& bOutSuccess)
 {
-	UGameplayTagsManager& TagManager = IGameplayTagsModule::Get().GetGameplayTagsManager();
+	UGameplayTagsManager& TagManager = IGameplayTagsModule::GetGameplayTagsManager();
 
 	uint8 bHasName = (TagName != NAME_None);
 	uint8 bHasNetIndex = 0;
diff --git a/Engine/Source/Runtime/GameplayTags/Private/GameplayTagsModule.cpp b/Engine/Source/Runtime/GameplayTags/Private/GameplayTagsModule.cpp
--- a/Engine/Source/Runtime/GameplayTags/Private/GameplayTagsModule.cpp
+++ b/Engine/Source/Runtime/GameplayTags/Private/GameplayTagsModule.cpp
@@ -8,18 +8,6 @@ class FGameplayTagsModule : public IGameplayTagsModule
 	virtual void StartupModule() override;
 	virtual void ShutdownModule() override;
 	// End IModuleInterface
-
-	// Gets the UGameplayTagsManager manager
-	virtual UGameplayTagsManager& GetGameplayTagsManager() override
-	{
-		check(GameplayTagsManager);
-		return *GameplayTagsManager;
-	}
-
-private:
-
-	// Stores the UGameplayTagsManager
-	UGameplayTagsManager* GameplayTagsManager;
 };
 
 IMPLEMENT_MODULE( FGameplayTagsModule, GameplayTags )
@@ -28,18 +16,18 @@ DEFINE_LOG_CATEGORY(LogGameplayTags);
 void FGameplayTagsModule::StartupModule()
 {
 	// This code will execute after your module is loaded into memory (but after global variables are initialized, of course.)
-	GameplayTagsManager = NewObject<UGameplayTagsManager>(GetTransientPackage(), NAME_None, RF_RootSet);
+	GGameplayTagsManager = NewObject<UGameplayTagsManager>(GetTransientPackage(), NAME_None, RF_RootSet);
 
 	TArray<FString> GameplayTagTables;
 	GConfig->GetArray(TEXT("GameplayTags"), TEXT("GameplayTagTableList"), GameplayTagTables, GEngineIni);
 
-	GameplayTagsManager->LoadGameplayTagTables(GameplayTagTables);
-	GameplayTagsManager->ConstructGameplayTagTree();
+	GGameplayTagsManager->LoadGameplayTagTables(GameplayTagTables);
+	GGameplayTagsManager->ConstructGameplayTagTree();
 }
 
 void FGameplayTagsModule::ShutdownModule()
 {
 	// This function may be called during shutdown to clean up your module.  For modules that support dynamic reloading,
 	// we call this function before unloading the module.
-	GameplayTagsManager = NULL;
+	GGameplayTagsManager = NULL;
 }
diff --git a/Engine/Source/Runtime/GameplayTags/Public/GameplayTagsModule.h b/Engine/Source/Runtime/GameplayTags/Public/GameplayTagsModule.h
--- a/Engine/Source/Runtime/GameplayTags/Public/GameplayTagsModule.h
+++ b/Engine/Source/Runtime/GameplayTags/Public/GameplayTagsModule.h
@@ -6,6 +6,8 @@
 #include "GameplayTagsManager.h"
 
 
+class UGameplayTagsManager;
+
 /**
  * The public interface to this module
  */
@@ -37,8 +39,6 @@ public:
 		return FModuleManager::Get().IsModuleLoaded(GameplayTagModuleName);
 	}
 
-	virtual UGameplayTagsManager& GetGameplayTagsManager() = 0;
-
 	/**
 	 * Helper function to request a gameplay tag by name
 	 * 
@@ -51,5 +51,18 @@ public:
 		IGameplayTagsModule& GameplayTagsModule = IGameplayTagsModule::Get();
 		return GameplayTagsModule.GetGameplayTagsManager().RequestGameplayTag(InTagName, ErrorIfNotFound);
 	}
+
+
+	// Gets the UGameplayTagsManager manager
+	static UGameplayTagsManager& GetGameplayTagsManager()
+	{
+		if (GGameplayTagsManager == nullptr)
+		{
+			IGameplayTagsModule::Get();
+		}
+		
+		return *GGameplayTagsManager;
+	}
+
 };
 
