UE4 - Optimized most of the low level gameplay tag queries.

diff --git a/Engine/Source/Runtime/GameplayTags/Classes/GameplayTagContainer.h b/Engine/Source/Runtime/GameplayTags/Classes/GameplayTagContainer.h
--- a/Engine/Source/Runtime/GameplayTags/Classes/GameplayTagContainer.h
+++ b/Engine/Source/Runtime/GameplayTags/Classes/GameplayTagContainer.h
@@ -4,6 +4,15 @@
 
 #include "GameplayTagContainer.generated.h"
 
+DECLARE_STATS_GROUP(TEXT("Gameplay Tags"), STATGROUP_GameplayTags, STATCAT_Advanced);
+
+DECLARE_CYCLE_STAT_EXTERN(TEXT("FGameplayTagContainer::HasTag"), STAT_FGameplayTagContainer_HasTag, STATGROUP_GameplayTags, GAMEPLAYTAGS_API);
+DECLARE_CYCLE_STAT_EXTERN(TEXT("FGameplayTagContainer::DoesTagContainerMatch"), STAT_FGameplayTagContainer_DoesTagContainerMatch, STATGROUP_GameplayTags, GAMEPLAYTAGS_API);
+DECLARE_CYCLE_STAT_EXTERN(TEXT("UGameplayTagsManager::GameplayTagsMatch"), STAT_UGameplayTagsManager_GameplayTagsMatch, STATGROUP_GameplayTags, GAMEPLAYTAGS_API);
+
+
+#define CHECK_TAG_OPTIMIZATIONS (0)
+
 UENUM(BlueprintType)
 namespace EGameplayTagMatchType
 {
@@ -30,20 +39,21 @@ struct GAMEPLAYTAGS_API FGameplayTag
 	GENERATED_USTRUCT_BODY()
 
 	/** Constructors */
-	FGameplayTag();
-
+	FGameplayTag()
+	{
+	}
 	/** Operators */
-	bool operator==(FGameplayTag const& Other) const
+	FORCEINLINE_DEBUGGABLE bool operator==(FGameplayTag const& Other) const
 	{
 		return TagName == Other.TagName;
 	}
 
-	bool operator!=(FGameplayTag const& Other) const
+	FORCEINLINE_DEBUGGABLE bool operator!=(FGameplayTag const& Other) const
 	{
 		return TagName != Other.TagName;
 	}
 
-	bool operator<(FGameplayTag const& Other) const
+	FORCEINLINE_DEBUGGABLE bool operator<(FGameplayTag const& Other) const
 	{
 		return TagName < Other.TagName;
 	}
@@ -57,10 +67,40 @@ struct GAMEPLAYTAGS_API FGameplayTag
 	 * 
 	 * @return True if there is a match according to the specified match types; false if not
 	 */
-	bool Matches(TEnumAsByte<EGameplayTagMatchType::Type> MatchTypeOne, const FGameplayTag& Other, TEnumAsByte<EGameplayTagMatchType::Type> MatchTypeTwo) const;
-	
+	FORCEINLINE_DEBUGGABLE bool Matches(TEnumAsByte<EGameplayTagMatchType::Type> MatchTypeOne, const FGameplayTag& Other, TEnumAsByte<EGameplayTagMatchType::Type> MatchTypeTwo) const
+	{
+		bool bResult;
+		if (MatchTypeOne == EGameplayTagMatchType::Explicit && MatchTypeTwo == EGameplayTagMatchType::Explicit)
+		{
+			bResult = TagName == Other.TagName;
+		}
+		else
+		{
+			bResult = ComplexMatches(MatchTypeOne, Other, MatchTypeTwo);
+		}
+#if CHECK_TAG_OPTIMIZATIONS
+		check(bResult == MatchesOriginal(MatchTypeOne, Other, MatchTypeTwo));
+#endif
+		return bResult;
+	}
+	/**
+	 * Check to see if two FGameplayTags match
+	 *
+	 * @param MatchTypeOne	How we compare this tag, Explicitly or a match with any parents as well
+	 * @param Other			The second tag to compare against
+	 * @param MatchTypeTwo	How we compare Other tag, Explicitly or a match with any parents as well
+	 * 
+	 * @return True if there is a match according to the specified match types; false if not
+	 */
+	bool ComplexMatches(TEnumAsByte<EGameplayTagMatchType::Type> MatchTypeOne, const FGameplayTag& Other, TEnumAsByte<EGameplayTagMatchType::Type> MatchTypeTwo) const;
+#if CHECK_TAG_OPTIMIZATIONS
+	bool MatchesOriginal(TEnumAsByte<EGameplayTagMatchType::Type> MatchTypeOne, const FGameplayTag& Other, TEnumAsByte<EGameplayTagMatchType::Type> MatchTypeTwo) const;
+#endif
 	/** Returns whether the tag is valid or not; Invalid tags are set to NAME_None and do not exist in the game-specific global dictionary */
-	bool IsValid() const;
+	FORCEINLINE bool IsValid() const
+	{
+		return (TagName != NAME_None);
+	}
 
 	/** Used so we can have a TMap of this struct*/
 	friend uint32 GetTypeHash(const FGameplayTag& Tag)
@@ -120,11 +160,27 @@ struct GAMEPLAYTAGS_API FGameplayTagContainer
 	GENERATED_USTRUCT_BODY()
 
 	/** Constructors */
-	FGameplayTagContainer();
-	FGameplayTagContainer(FGameplayTagContainer const& Other);
-	FGameplayTagContainer(const FGameplayTag& Tag);
-	FGameplayTagContainer(FGameplayTagContainer&& Other);
-	virtual ~FGameplayTagContainer() {}
+	FGameplayTagContainer()
+	{
+	}
+	FGameplayTagContainer(FGameplayTagContainer const& Other)
+	{
+		*this = Other;
+	}
+
+	FGameplayTagContainer(const FGameplayTag& Tag)
+	{
+		AddTag(Tag);
+	}
+
+	FGameplayTagContainer(FGameplayTagContainer&& Other)
+		: GameplayTags(MoveTemp(Other.GameplayTags))
+	{
+
+	}
+	~FGameplayTagContainer()
+	{
+	}
 
 	/** Assignment/Equality operators */
 	FGameplayTagContainer& operator=(FGameplayTagContainer const& Other);
@@ -155,7 +211,14 @@ struct GAMEPLAYTAGS_API FGameplayTagContainer
 	 *
 	 * @return True if this container has ANY the tags of the passed in container
 	 */
-	bool MatchesAny(const FGameplayTagContainer& Other, bool bCountEmptyAsMatch) const;
+	FORCEINLINE_DEBUGGABLE bool MatchesAny(const FGameplayTagContainer& Other, bool bCountEmptyAsMatch) const
+	{
+		if (Other.Num() == 0)
+		{
+			return bCountEmptyAsMatch;
+		}
+		return DoesTagContainerMatch(Other, EGameplayTagMatchType::IncludeParentTags, EGameplayTagMatchType::Explicit, EGameplayContainerMatchType::Any);
+	}
 
 	/**
 	* Checks if this container matches ALL of the tags in the specified container. Performs matching by expanding this container out to
@@ -166,7 +229,14 @@ struct GAMEPLAYTAGS_API FGameplayTagContainer
 	* 
 	* @return True if this container has ALL the tags of the passed in container
 	*/
-	bool MatchesAll(const FGameplayTagContainer& Other, bool bCountEmptyAsMatch) const;
+	FORCEINLINE_DEBUGGABLE bool MatchesAll(const FGameplayTagContainer& Other, bool bCountEmptyAsMatch) const
+	{
+		if (Other.Num() == 0)
+		{
+			return bCountEmptyAsMatch;
+		}
+		return DoesTagContainerMatch(Other, EGameplayTagMatchType::IncludeParentTags, EGameplayTagMatchType::Explicit, EGameplayContainerMatchType::All);
+	}
 
 	/** 
 	 * Checks if this container matches the given query.
@@ -186,7 +256,36 @@ struct GAMEPLAYTAGS_API FGameplayTagContainer
 	 * 
 	 * @return True if the tag is in the container, false if it is not
 	 */
-	bool HasTag(FGameplayTag const& TagToCheck, TEnumAsByte<EGameplayTagMatchType::Type> TagMatchType, TEnumAsByte<EGameplayTagMatchType::Type> TagToCheckMatchType) const;
+	FORCEINLINE_DEBUGGABLE bool HasTag(FGameplayTag const& TagToCheck, TEnumAsByte<EGameplayTagMatchType::Type> TagMatchType, TEnumAsByte<EGameplayTagMatchType::Type> TagToCheckMatchType) const
+	{
+		SCOPE_CYCLE_COUNTER(STAT_FGameplayTagContainer_HasTag);
+		bool bResult;
+		if (TagMatchType == EGameplayTagMatchType::Explicit && TagToCheckMatchType == EGameplayTagMatchType::Explicit)
+		{
+			bResult = GameplayTags.Contains(TagToCheck);
+		}
+		else
+		{
+			bResult = ComplexHasTag(TagToCheck, TagMatchType, TagToCheckMatchType);
+		}
+#if CHECK_TAG_OPTIMIZATIONS
+		check(bResult == HasTagOriginal(TagToCheck, TagMatchType, TagToCheckMatchType));
+#endif
+		return bResult;
+	}
+	/**
+	 * Determine if the container has the specified tag
+	 * 
+	 * @param TagToCheck			Tag to check if it is present in the container
+	 * @param TagMatchType			Type of match to use for the tags in this container
+	 * @param TagToCheckMatchType	Type of match to use for the TagToCheck Param
+	 * 
+	 * @return True if the tag is in the container, false if it is not
+	 */
+	bool ComplexHasTag(FGameplayTag const& TagToCheck, TEnumAsByte<EGameplayTagMatchType::Type> TagMatchType, TEnumAsByte<EGameplayTagMatchType::Type> TagToCheckMatchType) const;
+#if CHECK_TAG_OPTIMIZATIONS
+	bool HasTagOriginal(FGameplayTag const& TagToCheck, TEnumAsByte<EGameplayTagMatchType::Type> TagMatchType, TEnumAsByte<EGameplayTagMatchType::Type> TagToCheckMatchType) const;
+#endif
 
 	/** 
 	 * Adds all the tags from one container to this container 
@@ -287,7 +386,44 @@ struct GAMEPLAYTAGS_API FGameplayTagContainer
 		
 	bool NetSerialize(FArchive& Ar, class UPackageMap* Map, bool& bOutSuccess);
 
+	/**
+	* Returns true if the tags in this container match the tags in OtherContainer for the specified matching types.
+	*
+	* @param OtherContainer		The Container to filter against
+	* @param TagMatchType			Type of match to use for the tags in this container
+	* @param OtherTagMatchType		Type of match to use for the tags in the OtherContainer param
+	* @param ContainerMatchType	Type of match to use for filtering
+	*
+	* @return Returns true if ContainerMatchType is Any and any of the tags in OtherContainer match the tags in this or ContainerMatchType is All and all of the tags in OtherContainer match at least one tag in this. Returns false otherwise.
+	*/
+	FORCEINLINE_DEBUGGABLE bool DoesTagContainerMatch(const FGameplayTagContainer& OtherContainer, TEnumAsByte<EGameplayTagMatchType::Type> TagMatchType, TEnumAsByte<EGameplayTagMatchType::Type> OtherTagMatchType, EGameplayContainerMatchType ContainerMatchType) const
+	{
+		SCOPE_CYCLE_COUNTER(STAT_FGameplayTagContainer_DoesTagContainerMatch);
+		bool bResult;
+		if (ContainerMatchType == EGameplayContainerMatchType::Any && TagMatchType == EGameplayTagMatchType::Explicit && OtherTagMatchType == EGameplayTagMatchType::Explicit)
+		{
+			bResult = false;
+			for (TArray<FGameplayTag>::TConstIterator OtherIt(OtherContainer.GameplayTags); OtherIt; ++OtherIt)
+			{
+				if (GameplayTags.Contains(*OtherIt))
+				{
+					bResult = true;
+					break;
+				}
+			}			
+		}
+		else
+		{
+			bResult = DoesTagContainerMatchComplex(OtherContainer, TagMatchType, OtherTagMatchType, ContainerMatchType);
+		}
+#if CHECK_TAG_OPTIMIZATIONS
+		// the complex version is the original
+		check(bResult == DoesTagContainerMatchComplex(OtherContainer, TagMatchType, OtherTagMatchType, ContainerMatchType));
+#endif
+		return bResult;
+	}
 protected:
+
 	/**
 	* Returns true if the tags in this container match the tags in OtherContainer for the specified matching types.
 	*
@@ -298,7 +434,7 @@ protected:
 	*
 	* @return Returns true if ContainerMatchType is Any and any of the tags in OtherContainer match the tags in this or ContainerMatchType is All and all of the tags in OtherContainer match at least one tag in this. Returns false otherwise.
 	*/
-	bool DoesTagContainerMatch(const FGameplayTagContainer& OtherContainer, TEnumAsByte<EGameplayTagMatchType::Type> TagMatchType, TEnumAsByte<EGameplayTagMatchType::Type> OtherTagMatchType, EGameplayContainerMatchType ContainerMatchType) const;
+	bool DoesTagContainerMatchComplex(const FGameplayTagContainer& OtherContainer, TEnumAsByte<EGameplayTagMatchType::Type> TagMatchType, TEnumAsByte<EGameplayTagMatchType::Type> OtherTagMatchType, EGameplayContainerMatchType ContainerMatchType) const;
 
 	/** Array of gameplay tags */
 	UPROPERTY(BlueprintReadWrite, Category=GameplayTags)
diff --git a/Engine/Source/Runtime/GameplayTags/Classes/GameplayTagsManager.h b/Engine/Source/Runtime/GameplayTags/Classes/GameplayTagsManager.h
--- a/Engine/Source/Runtime/GameplayTags/Classes/GameplayTagsManager.h
+++ b/Engine/Source/Runtime/GameplayTags/Classes/GameplayTagsManager.h
@@ -106,6 +106,8 @@ private:
 
 	/** Owner gameplay tag node, if any */
 	TWeakPtr<FGameplayTagNode> ParentNode;
+
+	FGameplayTagContainer Parents;
 	
 	/** Net Index of this node */
 	FGameplayTagNetIndex NetIndex;
@@ -248,7 +250,47 @@ class GAMEPLAYTAGS_API UGameplayTagsManager : public UObject
 	 * 
 	 * @return true if there is a match
 	 */
-	bool GameplayTagsMatch(const FGameplayTag& GameplayTagOne, TEnumAsByte<EGameplayTagMatchType::Type> MatchTypeOne, const FGameplayTag& GameplayTagTwo, TEnumAsByte<EGameplayTagMatchType::Type> MatchTypeTwo) const;
+	FORCEINLINE_DEBUGGABLE bool GameplayTagsMatch(const FGameplayTag& GameplayTagOne, TEnumAsByte<EGameplayTagMatchType::Type> MatchTypeOne, const FGameplayTag& GameplayTagTwo, TEnumAsByte<EGameplayTagMatchType::Type> MatchTypeTwo) const
+	{
+		SCOPE_CYCLE_COUNTER(STAT_UGameplayTagsManager_GameplayTagsMatch);
+		bool bResult;
+		if (MatchTypeOne == EGameplayTagMatchType::Explicit && MatchTypeTwo == EGameplayTagMatchType::Explicit)
+		{
+			bResult = GameplayTagOne == GameplayTagTwo;
+		}
+		else
+		{
+			bResult = ComplexGameplayTagsMatch(GameplayTagOne, MatchTypeOne, GameplayTagTwo, MatchTypeTwo);
+		}
+#if CHECK_TAG_OPTIMIZATIONS
+		check(bResult == GameplayTagsMatchOriginal(GameplayTagOne, MatchTypeOne, GameplayTagTwo, MatchTypeTwo));
+#endif
+		return bResult;
+	}
+
+	/**
+	 * Same as above except MatchTypeOne == EGameplayTagMatchType::Explicit && MatchTypeTwo == EGameplayTagMatchType::Explicit is not tested or handled
+	 */
+	bool ComplexGameplayTagsMatch(const FGameplayTag& GameplayTagOne, TEnumAsByte<EGameplayTagMatchType::Type> MatchTypeOne, const FGameplayTag& GameplayTagTwo, TEnumAsByte<EGameplayTagMatchType::Type> MatchTypeTwo) const;
+#if CHECK_TAG_OPTIMIZATIONS
+	bool GameplayTagsMatchOriginal(const FGameplayTag& GameplayTagOne, TEnumAsByte<EGameplayTagMatchType::Type> MatchTypeOne, const FGameplayTag& GameplayTagTwo, TEnumAsByte<EGameplayTagMatchType::Type> MatchTypeTwo) const;
+#endif
+
+	/**
+	 * Helper function for GameplayTagsMatch to get a container with all parents.
+	 * @param GameplayTag		tag to get parents of
+	 * @return					pointer to container of parents, if any
+	 */
+	FORCEINLINE_DEBUGGABLE const FGameplayTagContainer* GetAllParentsContainer(const FGameplayTag& GameplayTag)
+	{
+		const TSharedPtr<FGameplayTagNode>* TagNode = GameplayTagNodeMap.Find(GameplayTag);
+		if (TagNode)
+		{
+			return &((*TagNode)->Parents);
+		}
+		return nullptr;
+	}
+
 
 	/** Event for when assets are added to the registry */
 	DECLARE_EVENT(UGameplayTagsManager, FGameplayTagTreeChanged);
diff --git a/Engine/Source/Runtime/GameplayTags/Private/GameplayTagContainer.cpp b/Engine/Source/Runtime/GameplayTags/Private/GameplayTagContainer.cpp
--- a/Engine/Source/Runtime/GameplayTags/Private/GameplayTagContainer.cpp
+++ b/Engine/Source/Runtime/GameplayTags/Private/GameplayTagContainer.cpp
@@ -5,6 +5,13 @@
 const FGameplayTagContainer FGameplayTagContainer::EmptyContainer;
 const FGameplayTagQuery FGameplayTagQuery::EmptyQuery;
 
+DEFINE_STAT(STAT_FGameplayTagContainer_HasTag);
+DEFINE_STAT(STAT_FGameplayTagContainer_DoesTagContainerMatch);
+DEFINE_STAT(STAT_UGameplayTagsManager_GameplayTagsMatch);
+
+
+
+
 /** Helper class to parse/eval query token streams. */
 class FQueryEvaluator
 {
@@ -396,26 +403,6 @@ bool FQueryEvaluator::EvalExpr(FGameplayTagContainer const& Tags, bool bSkip)
 }
 
 
-
-FGameplayTagContainer::FGameplayTagContainer()
-{}
-
-FGameplayTagContainer::FGameplayTagContainer(FGameplayTagContainer const& Other)
-{
-	*this = Other;
-}
-
-FGameplayTagContainer::FGameplayTagContainer(const FGameplayTag& Tag)
-{
-	AddTag(Tag);
-}
-
-FGameplayTagContainer::FGameplayTagContainer(FGameplayTagContainer&& Other)
-	: GameplayTags(MoveTemp(Other.GameplayTags))
-{
-	
-}
-
 FGameplayTagContainer& FGameplayTagContainer::operator=(FGameplayTagContainer const& Other)
 {
 	// Guard against self-assignment
@@ -453,8 +440,39 @@ bool FGameplayTagContainer::operator!=(FGameplayTagContainer const& Other) const
 	return Filter(Other, EGameplayTagMatchType::Explicit, EGameplayTagMatchType::Explicit).Num() != this->Num();
 }
 
-bool FGameplayTagContainer::HasTag(FGameplayTag const& TagToCheck, TEnumAsByte<EGameplayTagMatchType::Type> TagMatchType, TEnumAsByte<EGameplayTagMatchType::Type> TagToCheckMatchType) const
+bool FGameplayTagContainer::ComplexHasTag(FGameplayTag const& TagToCheck, TEnumAsByte<EGameplayTagMatchType::Type> TagMatchType, TEnumAsByte<EGameplayTagMatchType::Type> TagToCheckMatchType) const
+{
+	check(TagMatchType != EGameplayTagMatchType::Explicit || TagToCheckMatchType != EGameplayTagMatchType::Explicit);
+	UGameplayTagsManager& TagManager = IGameplayTagsModule::GetGameplayTagsManager();
+
+	if (TagMatchType != EGameplayTagMatchType::Explicit)
+	{
+		for (TArray<FGameplayTag>::TConstIterator It(this->GameplayTags); It; ++It)
+		{
+			const FGameplayTagContainer* Parents = TagManager.GetAllParentsContainer(*It);
+			if (Parents && Parents->HasTag(TagToCheck, EGameplayTagMatchType::Explicit, EGameplayTagMatchType::Explicit))
+			{
+				return true;
+			}
+		}
+	}
+	else
+	{
+		const FGameplayTagContainer* Parents = TagManager.GetAllParentsContainer(TagToCheck);
+		if (Parents && DoesTagContainerMatch(*Parents, EGameplayTagMatchType::Explicit, EGameplayTagMatchType::Explicit, EGameplayContainerMatchType::Any))
+		{
+			return true;
+		}
+
+	}
+	return false;
+}
+
+#if CHECK_TAG_OPTIMIZATIONS
+bool FGameplayTagContainer::HasTagOriginal(FGameplayTag const& TagToCheck, TEnumAsByte<EGameplayTagMatchType::Type> TagMatchType, TEnumAsByte<EGameplayTagMatchType::Type> TagToCheckMatchType) const
 {
+	SCOPE_CYCLE_COUNTER(STAT_FGameplayTagContainer_HasTag);
+
 	UGameplayTagsManager& TagManager = IGameplayTagsModule::GetGameplayTagsManager();
 	for (TArray<FGameplayTag>::TConstIterator It(this->GameplayTags); It; ++It)
 	{
@@ -465,6 +483,7 @@ bool FGameplayTagContainer::HasTag(FGameplayTag const& TagToCheck, TEnumAsByte<E
 	}
 	return false;
 }
+#endif
 
 bool FGameplayTagContainer::RemoveTagByExplicitName(const FName& TagName)
 {
@@ -514,7 +533,7 @@ FGameplayTagContainer FGameplayTagContainer::Filter(const FGameplayTagContainer&
 	return ResultContainer;
 }
 
-bool FGameplayTagContainer::DoesTagContainerMatch(const FGameplayTagContainer& OtherContainer, TEnumAsByte<EGameplayTagMatchType::Type> TagMatchType, TEnumAsByte<EGameplayTagMatchType::Type> OtherTagMatchType, EGameplayContainerMatchType ContainerMatchType) const
+bool FGameplayTagContainer::DoesTagContainerMatchComplex(const FGameplayTagContainer& OtherContainer, TEnumAsByte<EGameplayTagMatchType::Type> TagMatchType, TEnumAsByte<EGameplayTagMatchType::Type> OtherTagMatchType, EGameplayContainerMatchType ContainerMatchType) const
 {
 	UGameplayTagsManager& TagManager = IGameplayTagsModule::GetGameplayTagsManager();
 
@@ -549,27 +568,6 @@ bool FGameplayTagContainer::DoesTagContainerMatch(const FGameplayTagContainer& O
 	return ContainerMatchType == EGameplayContainerMatchType::All;
 }
 
-
-bool FGameplayTagContainer::MatchesAll(FGameplayTagContainer const& Other, bool bCountEmptyAsMatch) const
-{
-	if (Other.Num() == 0)
-	{
-		return bCountEmptyAsMatch;
-	}
-
-	return DoesTagContainerMatch(Other, EGameplayTagMatchType::IncludeParentTags, EGameplayTagMatchType::Explicit, EGameplayContainerMatchType::All);
-}
-
-bool FGameplayTagContainer::MatchesAny(FGameplayTagContainer const& Other, bool bCountEmptyAsMatch) const
-{
-	if (Other.Num() == 0)
-	{
-		return bCountEmptyAsMatch;
-	}
-
-	return DoesTagContainerMatch(Other, EGameplayTagMatchType::IncludeParentTags, EGameplayTagMatchType::Explicit, EGameplayContainerMatchType::Any);
-}
-
 bool FGameplayTagContainer::MatchesQuery(const FGameplayTagQuery& Query) const
 {
 	return Query.Matches(*this);
@@ -757,20 +755,17 @@ FText FGameplayTagContainer::ToMatchingText(EGameplayContainerMatchType MatchTyp
 	return FText::Format(MatchingDescription[DescriptionIndex], Arguments);
 }
 
-FGameplayTag::FGameplayTag()
-	: TagName(NAME_None)
-{
-}
-
-bool FGameplayTag::Matches(TEnumAsByte<EGameplayTagMatchType::Type> MatchTypeOne, const FGameplayTag& Other, TEnumAsByte<EGameplayTagMatchType::Type> MatchTypeTwo) const
+bool FGameplayTag::ComplexMatches(TEnumAsByte<EGameplayTagMatchType::Type> MatchTypeOne, const FGameplayTag& Other, TEnumAsByte<EGameplayTagMatchType::Type> MatchTypeTwo) const
 {
-	return IGameplayTagsModule::Get().GetGameplayTagsManager().GameplayTagsMatch(*this, MatchTypeOne, Other, MatchTypeTwo);
+	return IGameplayTagsModule::Get().GetGameplayTagsManager().ComplexGameplayTagsMatch(*this, MatchTypeOne, Other, MatchTypeTwo);
 }
 
-bool FGameplayTag::IsValid() const
+#if CHECK_TAG_OPTIMIZATIONS
+bool FGameplayTag::MatchesOriginal(TEnumAsByte<EGameplayTagMatchType::Type> MatchTypeOne, const FGameplayTag& Other, TEnumAsByte<EGameplayTagMatchType::Type> MatchTypeTwo) const
 {
-	return (TagName != NAME_None);
+	return IGameplayTagsModule::Get().GetGameplayTagsManager().GameplayTagsMatchOriginal(*this, MatchTypeOne, Other, MatchTypeTwo);
 }
+#endif
 
 FGameplayTag::FGameplayTag(FName Name)
 	: TagName(Name)
diff --git a/Engine/Source/Runtime/GameplayTags/Private/GameplayTagsManager.cpp b/Engine/Source/Runtime/GameplayTags/Private/GameplayTagsManager.cpp
--- a/Engine/Source/Runtime/GameplayTags/Private/GameplayTagsManager.cpp
+++ b/Engine/Source/Runtime/GameplayTags/Private/GameplayTagsManager.cpp
@@ -319,44 +319,44 @@ void UGameplayTagsManager::RedirectTagsForContainer(FGameplayTagContainer& Conta
 		{
 			NamesToRemove.Add(TagName);
 			if (NewTag->IsValid())
-			{
+				{
 				TagsToAdd.Add(NewTag);
 				DeprecatedTagNamesNotFoundInTagMap.Remove(NewTag->GetTagName());
+				}
+			}
 			}
-		}
-	}
 
 	// Add additional tags to the TagsToAdd set from the deprecated list if they weren't already added above
 	for (FName AdditionalDeprecatedTag : DeprecatedTagNamesNotFoundInTagMap)
-	{
+				{
 		const FGameplayTag* NewTag = TagRedirects.Find(AdditionalDeprecatedTag);
 		if (NewTag && NewTag->IsValid())
-		{
+					{
 			TagsToAdd.Add(NewTag);
-		}
-	}
+					}
+				}
 
 	// Remove all tags from the NamesToRemove set
 	for (FName RemoveName : NamesToRemove)
 	{
 		FGameplayTag OldTag = RequestGameplayTag(RemoveName, false);
 		if (OldTag.IsValid())
-		{
-			Container.RemoveTag(OldTag);
-		}
+				{
+				Container.RemoveTag(OldTag);
+			}
 		else
-		{
+			{
 			FGameplayTagRedirectHelper::RemoveTagByExplicitName(Container, RemoveName);
-		}
-	}
+				}
+			}
 
 	// Add all tags from the TagsToAdd set
 	for (const FGameplayTag* AddTag : TagsToAdd)
 	{
 		check(AddTag);
 		Container.AddTag(*AddTag);
+		}
 	}
-}
 
 void UGameplayTagsManager::PopulateTreeFromDataTable(class UDataTable* InTable)
 {
@@ -442,12 +442,19 @@ int32 UGameplayTagsManager::InsertTagIntoNodeArray(FName Tag, TWeakPtr<FGameplay
 
 		{
 			FScopeLock Lock(&GameplayTagMapCritical);
-		GameplayTagMap.Add(TagNode->GetCompleteTag(), GameplayTag);
+			GameplayTagMap.Add(TagNode->GetCompleteTag(), GameplayTag);
 		}
 
 		{
 			FScopeLock Lock(&GameplayTagNodeMapCritical);
 		GameplayTagNodeMap.Add(GameplayTag, TagNode);
+
+		TSet<FName> ParentSet;
+		GetAllParentNodeNames(ParentSet, TagNode);
+		for (auto Name : ParentSet)
+		{
+			TagNode->Parents.AddTagFast(FGameplayTag(Name));
+		}
 	}
 	}
 	else if (NodeArray[InsertionIdx]->CategoryDescription.IsEmpty() && !CategoryDescription.IsEmpty())
@@ -739,11 +746,51 @@ void UGameplayTagsManager::AddChildrenTags(FGameplayTagContainer& TagContainer,
 	}
 }
 
-bool UGameplayTagsManager::GameplayTagsMatch(const FGameplayTag& GameplayTagOne, TEnumAsByte<EGameplayTagMatchType::Type> MatchTypeOne, const FGameplayTag& GameplayTagTwo, TEnumAsByte<EGameplayTagMatchType::Type> MatchTypeTwo) const
+DECLARE_CYCLE_STAT(TEXT("UGameplayTagsManager::ComplexGameplayTagsMatch ParentsParents"), STAT_UGameplayTagsManager_GameplayTagsMatchParentsParents, STATGROUP_GameplayTags);
+
+bool UGameplayTagsManager::ComplexGameplayTagsMatch(const FGameplayTag& GameplayTagOne, TEnumAsByte<EGameplayTagMatchType::Type> MatchTypeOne, const FGameplayTag& GameplayTagTwo, TEnumAsByte<EGameplayTagMatchType::Type> MatchTypeTwo) const
 {
-	//fixme, why is there a critical section here
-	FScopeLock Lock(&GameplayTagNodeMapCritical);
-	check(IsInGameThread());
+	check(MatchTypeOne != EGameplayTagMatchType::Explicit || MatchTypeTwo != EGameplayTagMatchType::Explicit);
+
+	if (MatchTypeOne == EGameplayTagMatchType::IncludeParentTags && MatchTypeTwo == EGameplayTagMatchType::IncludeParentTags)
+	{
+		SCOPE_CYCLE_COUNTER(STAT_UGameplayTagsManager_GameplayTagsMatchParentsParents);
+		const TSharedPtr<FGameplayTagNode>* TagNodeOne = GameplayTagNodeMap.Find(GameplayTagOne);
+		const TSharedPtr<FGameplayTagNode>* TagNodeTwo = GameplayTagNodeMap.Find(GameplayTagTwo);
+		if (TagNodeOne && TagNodeTwo)
+		{
+			return (*TagNodeOne)->Parents.DoesTagContainerMatch((*TagNodeTwo)->Parents, EGameplayTagMatchType::Explicit, EGameplayTagMatchType::Explicit, EGameplayContainerMatchType::Any);
+		}
+	}
+	else
+	{
+		checkSlow(MatchTypeOne != EGameplayTagMatchType::Explicit || MatchTypeTwo != EGameplayTagMatchType::Explicit);
+
+		if (MatchTypeOne == EGameplayTagMatchType::IncludeParentTags)
+		{
+			const TSharedPtr<FGameplayTagNode>* TagNodeOne = GameplayTagNodeMap.Find(GameplayTagOne);
+			if (TagNodeOne)
+			{
+				return (*TagNodeOne)->Parents.HasTag(GameplayTagTwo, EGameplayTagMatchType::Explicit, EGameplayTagMatchType::Explicit);
+			}
+		}
+		else
+		{
+			const TSharedPtr<FGameplayTagNode>* TagNodeTwo = GameplayTagNodeMap.Find(GameplayTagTwo);
+			if (TagNodeTwo)
+			{
+				return (*TagNodeTwo)->Parents.HasTag(GameplayTagOne, EGameplayTagMatchType::Explicit, EGameplayTagMatchType::Explicit);
+			}
+		}
+	}
+	return false;
+}
+
+#if CHECK_TAG_OPTIMIZATIONS
+
+bool UGameplayTagsManager::GameplayTagsMatchOriginal(const FGameplayTag& GameplayTagOne, TEnumAsByte<EGameplayTagMatchType::Type> MatchTypeOne, const FGameplayTag& GameplayTagTwo, TEnumAsByte<EGameplayTagMatchType::Type> MatchTypeTwo) const
+{
+	SCOPE_CYCLE_COUNTER(STAT_UGameplayTagsManager_GameplayTagsMatch);
 	if (MatchTypeOne == EGameplayTagMatchType::Explicit && MatchTypeTwo == EGameplayTagMatchType::Explicit)
 	{
 		const TSharedPtr<FGameplayTagNode>* TagNode1 = GameplayTagNodeMap.Find(GameplayTagOne);
@@ -811,6 +858,8 @@ bool UGameplayTagsManager::GameplayTagsMatch(const FGameplayTag& GameplayTagOne,
 	return bResult;
 }
 
+#endif
+
 void UGameplayTagsManager::GetAllParentNodeNames(TSet<FName>& NamesList, const TSharedPtr<FGameplayTagNode> GameplayTag) const
 {
 	NamesList.Add(GameplayTag->GetCompleteTag());
diff --git a/Engine/Source/Runtime/GameplayTags/Public/GameplayTagsModule.h b/Engine/Source/Runtime/GameplayTags/Public/GameplayTagsModule.h
--- a/Engine/Source/Runtime/GameplayTags/Public/GameplayTagsModule.h
+++ b/Engine/Source/Runtime/GameplayTags/Public/GameplayTagsModule.h
@@ -46,7 +46,7 @@ public:
 	 * 
 	 * @return Gameplay tag associated with the specified name. Will be marked invalid if tag not found
 	 */
-	static FGameplayTag RequestGameplayTag(FName InTagName, bool ErrorIfNotFound=true)
+	FORCEINLINE_DEBUGGABLE static FGameplayTag RequestGameplayTag(FName InTagName, bool ErrorIfNotFound=true)
 	{
 		IGameplayTagsModule& GameplayTagsModule = IGameplayTagsModule::Get();
 		return GameplayTagsModule.GetGameplayTagsManager().RequestGameplayTag(InTagName, ErrorIfNotFound);
@@ -54,7 +54,7 @@ public:
 
 
 	// Gets the UGameplayTagsManager manager
-	static UGameplayTagsManager& GetGameplayTagsManager()
+	FORCEINLINE_DEBUGGABLE static UGameplayTagsManager& GetGameplayTagsManager()
 	{
 		if (GGameplayTagsManager == nullptr)
 		{
