Fix FGameplayTagContainer not calling NetSerialize on its tags (it was serializing its tags as strings instead!)

--------
Integrated using branch UE4-To-UE4-Orion (reversed) of change#2686281 by [REDACTED_PERSON] on 2015/09/10 10:57:46.

diff --git a/Engine/Source/Runtime/GameplayTags/Private/GameplayTagContainer.cpp b/Engine/Source/Runtime/GameplayTags/Private/GameplayTagContainer.cpp
--- a/Engine/Source/Runtime/GameplayTags/Private/GameplayTagContainer.cpp
+++ b/Engine/Source/Runtime/GameplayTags/Private/GameplayTagContainer.cpp
@@ -713,7 +713,29 @@ FString FGameplayTagContainer::ToStringSimple() const
 
 bool FGameplayTagContainer::NetSerialize(FArchive& Ar, class UPackageMap* Map, bool& bOutSuccess)
 {
-	Ar << GameplayTags;
+	uint8 NumTags;
+	if (Ar.IsSaving())
+	{
+		NumTags = GameplayTags.Num();
+		Ar << NumTags;
+		for (FGameplayTag& Tag : GameplayTags)
+		{
+			Tag.NetSerialize(Ar, Map, bOutSuccess);
+		}
+	}
+	else
+	{
+		Ar << NumTags;
+		GameplayTags.Empty(NumTags);
+		GameplayTags.AddDefaulted(NumTags);
+		for (uint8 idx = 0; idx < NumTags; ++idx)
+		{
+			GameplayTags[idx].NetSerialize(Ar, Map, bOutSuccess);
+		}
+
+	}
+
+	bOutSuccess  = true;
 	return true;
 }
 
