Make AppendGameplayTagContainers non pure. UE-20863

--------
Integrated using branch UE4-To-UE4-Orion (reversed) of change#2692041 by [REDACTED_PERSON] on 2015/09/15 13:35:05.

diff --git a/Engine/Source/Runtime/GameplayTags/Classes/BlueprintGameplayTagLibrary.h b/Engine/Source/Runtime/GameplayTags/Classes/BlueprintGameplayTagLibrary.h
--- a/Engine/Source/Runtime/GameplayTags/Classes/BlueprintGameplayTagLibrary.h
+++ b/Engine/Source/Runtime/GameplayTags/Classes/BlueprintGameplayTagLibrary.h
@@ -92,7 +92,7 @@ class UBlueprintGameplayTagLibrary : public UBlueprintFunctionLibrary
 	* @param InTagContainer			The container to append.
 	* @param InOutTagContainer		The container that will be appended too.
 	*/
-	UFUNCTION(BlueprintPure, Category = "GameplayTags|Tag Container")
+	UFUNCTION(BlueprintCallable, Category = "GameplayTags|Tag Container")
 	static bool AppendGameplayTagContainers(const FGameplayTagContainer& InTagContainer, UPARAM(ref) FGameplayTagContainer& InOutTagContainer);
 
 	/**
