fix tag queries: FQueryEvaluator::EvalNoExprMatch, FQueryEvaluator::EvalAllExprMatch a

--------
Integrated using branch UE4-To-UE4-Orion (reversed) of change#2695065 by [REDACTED_PERSON] on 2015/09/17 09:26:29.

diff --git a/Engine/Source/Runtime/GameplayTags/Private/GameplayTagContainer.cpp b/Engine/Source/Runtime/GameplayTags/Private/GameplayTagContainer.cpp
--- a/Engine/Source/Runtime/GameplayTags/Private/GameplayTagContainer.cpp
+++ b/Engine/Source/Runtime/GameplayTags/Private/GameplayTagContainer.cpp
@@ -314,8 +314,8 @@ bool FQueryEvaluator::EvalAllExprMatch(FGameplayTagContainer const& Tags, bool b
 {
 	bool bShortCircuit = bSkip;
 
-	// assume false until proven otherwise
-	bool Result = false;
+	// assume true until proven otherwise
+	bool Result = true;
 
 	// parse exprset
 	int32 const NumExprs = GetToken();
@@ -329,10 +329,10 @@ bool FQueryEvaluator::EvalAllExprMatch(FGameplayTagContainer const& Tags, bool b
 		bool const bExprResult = EvalExpr(Tags, bShortCircuit);
 		if (bShortCircuit == false)
 		{
-			if (bExprResult == true)
+			if (bExprResult == false)
 			{
-				// one match is sufficient for true result
-				Result = true;
+				// one fail is sufficient for false result
+				Result = false;
 				bShortCircuit = true;
 			}
 		}
@@ -344,8 +344,8 @@ bool FQueryEvaluator::EvalNoExprMatch(FGameplayTagContainer const& Tags, bool bS
 {
 	bool bShortCircuit = bSkip;
 
-	// assume false until proven otherwise
-	bool Result = false;
+	// assume true until proven otherwise
+	bool Result = true;
 
 	// parse exprset
 	int32 const NumExprs = GetToken();
@@ -361,8 +361,8 @@ bool FQueryEvaluator::EvalNoExprMatch(FGameplayTagContainer const& Tags, bool bS
 		{
 			if (bExprResult == true)
 			{
-				// one match is sufficient for true result
-				Result = true;
+				// one match is sufficient for fail result
+				Result = false;
 				bShortCircuit = true;
 			}
 		}
