Copying //UE4/Dev-Framework to //UE4/Dev-Main (Source: //UE4/Dev-Framework @ 3182037)

diff --git a/Engine/Source/Runtime/GameplayTags/Classes/BlueprintGameplayTagLibrary.h b/Engine/Source/Runtime/GameplayTags/Classes/BlueprintGameplayTagLibrary.h
--- a/Engine/Source/Runtime/GameplayTags/Classes/BlueprintGameplayTagLibrary.h
+++ b/Engine/Source/Runtime/GameplayTags/Classes/BlueprintGameplayTagLibrary.h
@@ -16,17 +16,48 @@ class UBlueprintGameplayTagLibrary : public UBlueprintFunctionLibrary
 	GENERATED_UCLASS_BODY()
 
 	/**
-	 * Determine if the specified gameplay tags match, given the specified match types
+	 * Determine if TagOne matches against TagTwo
 	 * 
-	 * @param TagOne			First tag to check
-	 * @param TagTwo			Second tag to check
-	 * @param TagOneMatchType	Matching type to use on the first tag
-	 * @param TagTwoMatchType	Matching type to use on the second tag
+	 * @param TagOne			Tag to check for match
+	 * @param TagTwo			Tag to check match against
+	 * @param bExactMatch		If true, the tag has to be exactly present, if false then TagOne will include it's parent tags while matching			
 	 * 
-	 * @return True if the tags match, false if they do not
+	 * @return True if TagOne matches TagTwo
+	 */
+	UFUNCTION(BlueprintPure, Category="GameplayTags", meta = (Keywords = "DoGameplayTagsMatch"))
+	static bool MatchesTag(FGameplayTag TagOne, FGameplayTag TagTwo, bool bExactMatch);
+
+	/**
+	 * Determine if TagOne matches against any tag in OtherContainer
+	 * 
+	 * @param TagOne			Tag to check for match
+	 * @param OtherContainer	Container to check against.
+	 * @param bExactMatch		If true, the tag has to be exactly present, if false then TagOne will include it's parent tags while matching
+	 * 
+	 * @return True if TagOne matches any tags explicitly present in OtherContainer
 	 */
 	UFUNCTION(BlueprintPure, Category="GameplayTags")
-	static bool DoGameplayTagsMatch(const FGameplayTag& TagOne, const FGameplayTag& TagTwo, TEnumAsByte<EGameplayTagMatchType::Type> TagOneMatchType, TEnumAsByte<EGameplayTagMatchType::Type> TagTwoMatchType);
+	static bool MatchesAnyTags(FGameplayTag TagOne, const FGameplayTagContainer& OtherContainer, bool bExactMatch);
+
+	/** Returns true if the values are equal (A == B) */
+	UFUNCTION(BlueprintPure, meta=(DisplayName="Equal (GameplayTag)", CompactNodeTitle="=="), Category="GameplayTags")
+	static bool EqualEqual_GameplayTag( FGameplayTag A, FGameplayTag B );
+	
+	/** Returns true if the values are not equal (A != B) */
+	UFUNCTION(BlueprintPure, meta=(DisplayName="NotEqual (GameplayTag)", CompactNodeTitle="!="), Category="GameplayTags")
+	static bool NotEqual_GameplayTag( FGameplayTag A, FGameplayTag B );
+
+	/** Returns true if the passed in gameplay tag is non-null */
+	UFUNCTION(BlueprintPure, Category = "GameplayTags")
+	static bool IsGameplayTagValid(FGameplayTag GameplayTag);
+
+	/** Returns FName of this tag */
+	UFUNCTION(BlueprintPure, Category = "GameplayTags")
+	static FName GetTagName(const FGameplayTag& GameplayTag);
+
+	/** Creates a literal FGameplayTag */
+	UFUNCTION(BlueprintPure, Category = "GameplayTags")
+	static FGameplayTag MakeLiteralGameplayTag(FGameplayTag Value);
 
 	/**
 	 * Get the number of gameplay tags in the specified container
@@ -35,45 +66,44 @@ class UBlueprintGameplayTagLibrary : public UBlueprintFunctionLibrary
 	 * 
 	 * @return The number of tags in the specified container
 	 */
-	UFUNCTION(BlueprintPure, Category="GameplayTags|Tag Container")
+	UFUNCTION(BlueprintPure, Category="GameplayTags")
 	static int32 GetNumGameplayTagsInContainer(const FGameplayTagContainer& TagContainer);
 
 	/**
-	 * Check if the specified tag container has the specified tag, using the specified tag matching types
+	 * Check if the tag container has the specified tag
 	 *
-	 * @param TagContainer				Container to check for the tag
-	 * @param ContainerTagsMatchType	Matching options to use for tags inside the container
-	 * @param Tag						Tag to check for in the container
-	 * @param TagMatchType				Matching option to use for the tag
+	 * @param TagContainer			Container to check for the tag
+	 * @param Tag					Tag to check for in the container
+	 * @param bExactMatch			If true, the tag has to be exactly present, if false then TagContainer will include it's parent tags while matching			
 	 *
 	 * @return True if the container has the specified tag, false if it does not
 	 */
-	UFUNCTION(BlueprintPure, Category="GameplayTags|Tag Container")
-	static bool DoesContainerHaveTag(const FGameplayTagContainer& TagContainer, TEnumAsByte<EGameplayTagMatchType::Type> ContainerTagsMatchType, const FGameplayTag& Tag, TEnumAsByte<EGameplayTagMatchType::Type> TagMatchType);
+	UFUNCTION(BlueprintPure, Category="GameplayTags", meta = (Keywords = "DoesContainerHaveTag"))
+	static bool HasTag(const FGameplayTagContainer& TagContainer, FGameplayTag Tag, bool bExactMatch);
 
 	/**
-	 * Check if the specified tag container matches ANY of the tags in the other container. Matching is performed by expanding the tag container out to include all of its parent tags as well.
+	 * Check if the specified tag container has ANY of the tags in the other container
 	 * 
 	 * @param TagContainer			Container to check if it matches any of the tags in the other container
-	 * @param OtherContainer		Container to check against
-	 * @param bCountEmptyAsMatch	If true, the other container will count as a match, even if it's empty
+	 * @param OtherContainer		Container to check against.
+	 * @param bExactMatch			If true, the tag has to be exactly present, if false then TagContainer will include it's parent tags while matching			
 	 * 
-	 * @return True if the container matches ANY of the tags in the other container
+	 * @return True if the container has ANY of the tags in the other container
 	 */
-	UFUNCTION(BlueprintPure, Category="GameplayTags|Tag Container")
-	static bool DoesContainerMatchAnyTagsInContainer(const FGameplayTagContainer& TagContainer, const FGameplayTagContainer& OtherContainer, bool bCountEmptyAsMatch);
+	UFUNCTION(BlueprintPure, Category="GameplayTags", meta = (Keywords = "DoesContainerMatchAnyTagsInContainer"))
+	static bool HasAnyTags(const FGameplayTagContainer& TagContainer, const FGameplayTagContainer& OtherContainer, bool bExactMatch);
 
 	/**
-	 * Check if the specified tag container matches ALL of the tags in the other container. Matching is performed by expanding the tag container out to include all of its parent tags as well.
+	 * Check if the specified tag container has ALL of the tags in the other container
 	 * 
 	 * @param TagContainer			Container to check if it matches all of the tags in the other container
-	 * @param OtherContainer		Container to check against
-	 * @param bCountEmptyAsMatch	If true, the other container will count as a match, even if it's empty
+	 * @param OtherContainer		Container to check against. If this is empty, the check will succeed
+	 * @param bExactMatch			If true, the tag has to be exactly present, if false then TagContainer will include it's parent tags while matching			
 	 * 
-	 * @return True if the container matches ALL of the tags in the other container
+	 * @return True if the container has ALL of the tags in the other container
 	 */
-	UFUNCTION(BlueprintPure, Category="GameplayTags|Tag Container")
-	static bool DoesContainerMatchAllTagsInContainer(const FGameplayTagContainer& TagContainer, const FGameplayTagContainer& OtherContainer, bool bCountEmptyAsMatch);
+	UFUNCTION(BlueprintPure, Category="GameplayTags", meta = (Keywords = "DoesContainerMatchAllTagsInContainer"))
+	static bool HasAllTags(const FGameplayTagContainer& TagContainer, const FGameplayTagContainer& OtherContainer, bool bExactMatch);
 
 	/**
 	 * Check if the specified tag container matches the given Tag Query.
@@ -83,33 +113,55 @@ class UBlueprintGameplayTagLibrary : public UBlueprintFunctionLibrary
 	 * 
 	 * @return True if the container matches the query, false otherwise.
 	 */
-	UFUNCTION(BlueprintPure, Category = "GameplayTags|Tag Container")
+	UFUNCTION(BlueprintPure, Category = "GameplayTags")
 	static bool DoesContainerMatchTagQuery(const FGameplayTagContainer& TagContainer, const FGameplayTagQuery& TagQuery);
 
 	/**
-	* Takes 2 Gameplay Tag Containers Appends the In Container to the In/Out Container
-	*
-	* @param InTagContainer			The container to append.
-	* @param InOutTagContainer		The container that will be appended too.
-	*/
-	UFUNCTION(BlueprintCallable, Category = "GameplayTags|Tag Container")
-	static bool AppendGameplayTagContainers(const FGameplayTagContainer& InTagContainer, UPARAM(ref) FGameplayTagContainer& InOutTagContainer);
+	 * Takes 2 Gameplay Tag Containers Appends the In Container to the In/Out Container
+	 *
+	 * @param InOutTagContainer		The container that will be appended too.
+	 * @param Tag					The tag to add to the container
+	 */
+	UFUNCTION(BlueprintCallable, Category = "GameplayTags")
+	static void AddGameplayTag(UPARAM(ref) FGameplayTagContainer& InOutTagContainer, FGameplayTag Tag);
 
 	/**
-	 * Creates a literal FGameplayTag
-	 *
-	 * @param	Value	value to set the FGameplayTag to
+	* Takes 2 Gameplay Tag Containers Appends the In Container to the In/Out Container
 	 *
-	 * @return	The literal FGameplayTag
+	 * @param InOutTagContainer		The container that will be appended too.
+	 * @param InTagContainer		The container to append.
 	 */
+	UFUNCTION(BlueprintCallable, Category = "GameplayTags")
+	static void AppendGameplayTagContainers(UPARAM(ref) FGameplayTagContainer& InOutTagContainer, const FGameplayTagContainer& InTagContainer);
+
+	/** Returns true if the values are equal (A == B) */
+	UFUNCTION(BlueprintPure, meta=(DisplayName="Equal (GameplayTagContainer)", CompactNodeTitle="=="), Category="GameplayTags")
+	static bool EqualEqual_GameplayTagContainer( const FGameplayTagContainer& A, const FGameplayTagContainer& B );
+	
+	/** Returns true if the values are not equal (A != B) */
+	UFUNCTION(BlueprintPure, meta=(DisplayName="NotEqual (GameplayTagContainer)", CompactNodeTitle="!="), Category="GameplayTags")
+	static bool NotEqual_GameplayTagContainer( const FGameplayTagContainer& A, const FGameplayTagContainer& B );
+
+	/** Creates a literal FGameplayTagContainer */
 	UFUNCTION(BlueprintPure, Category = "GameplayTags")
-	static FGameplayTag MakeLiteralGameplayTag(FGameplayTag Value);
+	static FGameplayTagContainer MakeLiteralGameplayTagContainer(FGameplayTagContainer Value);
+
+	/** Creates a FGameplayTagContainer from the array of passed in tags */
+	UFUNCTION(BlueprintPure, Category = "GameplayTags")
+	static FGameplayTagContainer MakeGameplayTagContainerFromArray(const TArray<FGameplayTag>& GameplayTags);
+
+	/** Creates a FGameplayTagContainer containing a single tag */
+	UFUNCTION(BlueprintPure, Category = "GameplayTags")
+	static FGameplayTagContainer MakeGameplayTagContainerFromTag(FGameplayTag SingleTag);
 
+	/** Breaks tag container into explicit array of tags */
+	UFUNCTION(BlueprintPure, Category = "GameplayTags")
+	static void BreakGameplayTagContainer(const FGameplayTagContainer& GameplayTagContainer, TArray<FGameplayTag>& GameplayTags);
 
 	/**
 	 * Creates a literal FGameplayTagQuery
 	 *
-	 * @param	Value	value to set the FGameplayTagQuery to
+	 * @param	TagQuery	value to set the FGameplayTagQuery to
 	 *
 	 * @return	The literal FGameplayTagQuery
 	 */
@@ -125,27 +177,25 @@ class UBlueprintGameplayTagLibrary : public UBlueprintFunctionLibrary
 	 * @return True if the tagcontainer in the interface has all the tags inside the container.
 	 */
 	UFUNCTION(BlueprintPure, meta = (BlueprintInternalUseOnly = "TRUE"))
-	static bool HasAllMatchingGameplayTags(TScriptInterface<IGameplayTagAssetInterface> TagContainerInterface, const FGameplayTagContainer& OtherContainer, bool bCountEmptyAsMatch);
+	static bool HasAllMatchingGameplayTags(TScriptInterface<IGameplayTagAssetInterface> TagContainerInterface, const FGameplayTagContainer& OtherContainer);
 
 	/**
-	* Check if the specified tag container has the specified tag, using the specified tag matching types
-	*
-	* @param TagContainerInterface		An Interface to a tag container
-	* @param ContainerTagsMatchType		Matching options to use for tags inside the container
-	* @param Tag						Tag to check for in the container
-	* @param TagMatchType				Matching option to use for the tag
-	*
-	* @return True if the container has the specified tag, false if it does not
-	*/
+	 * Check if the specified tag container has the specified tag, using the specified tag matching types
+	 *
+	 * @param TagContainerInterface		An Interface to a tag container
+	 * @param Tag						Tag to check for in the container
+	 *
+	 * @return True if the container has the specified tag, false if it does not
+	 */
 	UFUNCTION(BlueprintPure, meta = (BlueprintInternalUseOnly = "TRUE"))
-	static bool DoesTagAssetInterfaceHaveTag(TScriptInterface<IGameplayTagAssetInterface> TagContainerInterface, TEnumAsByte<EGameplayTagMatchType::Type> ContainerTagsMatchType, const FGameplayTag& Tag, TEnumAsByte<EGameplayTagMatchType::Type> TagMatchType);
+	static bool DoesTagAssetInterfaceHaveTag(TScriptInterface<IGameplayTagAssetInterface> TagContainerInterface, FGameplayTag Tag);
 
 	/** Checks if a gameplay tag's name and a string are not equal to one another */
 	UFUNCTION(BlueprintPure, Category = PinOptions, meta = (BlueprintInternalUseOnly = "TRUE"))
 	static bool NotEqual_TagTag(FGameplayTag A, FString B);
 
 	/** Checks if a gameplay tag containers's name and a string are not equal to one another */
-	UFUNCTION(BlueprintPure, Category = PinOptions)
+	UFUNCTION(BlueprintPure, Category = PinOptions, meta = (BlueprintInternalUseOnly = "TRUE"))
 	static bool NotEqual_TagContainerTagContainer(FGameplayTagContainer A, FString B);
 	
 	/**
@@ -153,7 +203,7 @@ class UBlueprintGameplayTagLibrary : public UBlueprintFunctionLibrary
 	 *
 	 * @param TagContainer	The tag container to get the debug string from.
 	 */
-	UFUNCTION(BlueprintPure, Category = "GameplayTags|Tag Container")
+	UFUNCTION(BlueprintPure, Category = "GameplayTags")
 	static FString GetDebugStringFromGameplayTagContainer(const FGameplayTagContainer& TagContainer);
 
 	/**
@@ -162,18 +212,6 @@ class UBlueprintGameplayTagLibrary : public UBlueprintFunctionLibrary
 	 * @param GameplayTag	The tag to get the debug string from.
 	 */
 	UFUNCTION(BlueprintPure, Category = "GameplayTags")
-	static FString GetDebugStringFromGameplayTag(const FGameplayTag& GameplayTag);
+	static FString GetDebugStringFromGameplayTag(FGameplayTag GameplayTag);
 
-	/**
-	 * Check if the specified tag container has the specified tag, using the specified tag matching types
-	 *
-	 * @param TagContainer				Container to check for the tag
-	 * @param ContainerTagsMatchType	Matching options to use for tags inside the container
-	 * @param Tag						Tag to check for in the container
-	 * @param TagMatchType				Matching option to use for the tag
-	 *
-	 * @return True if the container has the specified tag, false if it does not
-	 */
-	UFUNCTION(BlueprintPure, Category="GameplayTags")
-	static bool IsGameplayTagValid(const FGameplayTag& TagContainer);
 };
diff --git a/Engine/Source/Runtime/GameplayTags/Classes/GameplayTagAssetInterface.h b/Engine/Source/Runtime/GameplayTags/Classes/GameplayTagAssetInterface.h
--- a/Engine/Source/Runtime/GameplayTags/Classes/GameplayTagAssetInterface.h
+++ b/Engine/Source/Runtime/GameplayTags/Classes/GameplayTagAssetInterface.h
@@ -38,22 +38,20 @@ class GAMEPLAYTAGS_API IGameplayTagAssetInterface
 	 * Check if the asset has gameplay tags that matches against all of the specified tags (expands to include parents of asset tags)
 	 * 
 	 * @param TagContainer			Tag container to check for a match
-	 * @param bCountEmptyAsMatch	If true, the parameter tag container will count as matching, even if it's empty
 	 * 
-	 * @return True if the asset has matches all of the gameplay tags
+	 * @return True if the asset has matches all of the gameplay tags, will be true if container is empty
 	 */
 	UFUNCTION(BlueprintCallable, Category=GameplayTags)
-	virtual bool HasAllMatchingGameplayTags(const FGameplayTagContainer& TagContainer, bool bCountEmptyAsMatch = true) const;
+	virtual bool HasAllMatchingGameplayTags(const FGameplayTagContainer& TagContainer) const;
 
 	/**
 	 * Check if the asset has gameplay tags that matches against any of the specified tags (expands to include parents of asset tags)
 	 * 
 	 * @param TagContainer			Tag container to check for a match
-	 * @param bCountEmptyAsMatch	If true, the parameter tag container will count as matching, even if it's empty
 	 * 
-	 * @return True if the asset has matches any of the gameplay tags
+	 * @return True if the asset has matches any of the gameplay tags, will be false if container is empty
 	 */
 	UFUNCTION(BlueprintCallable, Category=GameplayTags)
-	virtual bool HasAnyMatchingGameplayTags(const FGameplayTagContainer& TagContainer, bool bCountEmptyAsMatch = true) const;
+	virtual bool HasAnyMatchingGameplayTags(const FGameplayTagContainer& TagContainer) const;
 };
 
diff --git a/Engine/Source/Runtime/GameplayTags/Classes/GameplayTagContainer.h b/Engine/Source/Runtime/GameplayTags/Classes/GameplayTagContainer.h
--- a/Engine/Source/Runtime/GameplayTags/Classes/GameplayTagContainer.h
+++ b/Engine/Source/Runtime/GameplayTags/Classes/GameplayTagContainer.h
@@ -12,8 +12,7 @@
 DECLARE_CYCLE_STAT_EXTERN(TEXT("FGameplayTagContainer::DoesTagContainerMatch"), STAT_FGameplayTagContainer_DoesTagContainerMatch, STATGROUP_GameplayTags, GAMEPLAYTAGS_API);
 DECLARE_CYCLE_STAT_EXTERN(TEXT("UGameplayTagsManager::GameplayTagsMatch"), STAT_UGameplayTagsManager_GameplayTagsMatch, STATGROUP_GameplayTags, GAMEPLAYTAGS_API);
 
+struct FGameplayTagContainer;
 
-#define CHECK_TAG_OPTIMIZATIONS (0)
-
 UENUM(BlueprintType)
 namespace EGameplayTagMatchType
@@ -35,7 +34,8 @@
 typedef uint16 FGameplayTagNetIndex;
 #define INVALID_TAGNETINDEX MAX_uint16
 
-USTRUCT(BlueprintType)
+/** A single gameplay tag, which represents a hierarchical name of the form x.y that is registered in the GameplayTagsManager */
+USTRUCT(BlueprintType, meta = (HasNativeMake = "GameplayTags.BlueprintGameplayTagLibrary.MakeLiteralGameplayTag", HasNativeBreak = "GameplayTags.BlueprintGameplayTagLibrary.GetTagName"))
 struct GAMEPLAYTAGS_API FGameplayTag
 {
 	GENERATED_USTRUCT_BODY()
@@ -44,80 +44,116 @@
 	FGameplayTag()
 	{
 	}
+
+	/**
+	 * Gets the FGameplayTag that corresponds to the TagName
+	 *
+	 * @param TagName The Name of the tag to search for
+	 * 
+	 * @param ErrorIfNotfound: ensure() that tag exists.
+	 * 
+	 * @return Will return the corresponding FGameplayTag or an empty one if not found.
+	 */
+	static FGameplayTag RequestGameplayTag(FName TagName, bool ErrorIfNotFound=true);
+
 	/** Operators */
-	FORCEINLINE_DEBUGGABLE bool operator==(FGameplayTag const& Other) const
+	FORCEINLINE bool operator==(FGameplayTag const& Other) const
 	{
 		return TagName == Other.TagName;
 	}
 
-	FORCEINLINE_DEBUGGABLE bool operator!=(FGameplayTag const& Other) const
+	FORCEINLINE bool operator!=(FGameplayTag const& Other) const
 	{
 		return TagName != Other.TagName;
 	}
 
-	FORCEINLINE_DEBUGGABLE bool operator<(FGameplayTag const& Other) const
+	FORCEINLINE bool operator<(FGameplayTag const& Other) const
 	{
 		return TagName < Other.TagName;
 	}
-	
+
 	/**
-	 * Check to see if two FGameplayTags match
+	 * Determine if this tag matches TagToCheck, expanding our parent tags
+	 * "A.1".MatchesTag("A") will return True, "A".MatchesTag("A.1") will return False
+	 * If TagToCheck is not Valid it will always return False
+	 * 
+	 * @return True if this tag matches TagToCheck
+	 */
+	bool MatchesTag(const FGameplayTag& TagToCheck) const;
+
+	/**
+	 * Determine if TagToCheck is valid and exactly matches this tag
+	 * "A.1".MatchesTagExact("A") will return False
+	 * If TagToCheck is not Valid it will always return False
+	 * 
+	 * @return True if TagToCheck is Valid and is exactly this tag
+	 */
+	FORCEINLINE bool MatchesTagExact(const FGameplayTag& TagToCheck) const
+	{
+		if (!TagToCheck.IsValid())
+		{
+			return false;
+		}
+		// Only check check explicit tag list
+		return TagName == TagToCheck.TagName;
+	}
+
+	/**
+	 * Check to see how closely two FGameplayTags match. Higher values indicate more matching terms in the tags.
 	 *
-	 * @param MatchTypeOne	How we compare this tag, Explicitly or a match with any parents as well
-	 * @param Other			The second tag to compare against
-	 * @param MatchTypeTwo	How we compare Other tag, Explicitly or a match with any parents as well
-	 * 
-	 * @return True if there is a match according to the specified match types; false if not
+	 * @param TagToCheck	Tag to match against
+	 *
+	 * @return The depth of the match, higher means they are closer to an exact match
 	 */
-	FORCEINLINE_DEBUGGABLE bool Matches(TEnumAsByte<EGameplayTagMatchType::Type> MatchTypeOne, const FGameplayTag& Other, TEnumAsByte<EGameplayTagMatchType::Type> MatchTypeTwo) const
-	{
-		bool bResult;
-		if (MatchTypeOne == EGameplayTagMatchType::Explicit && MatchTypeTwo == EGameplayTagMatchType::Explicit)
-		{
-			bResult = TagName == Other.TagName;
-		}
-		else
-		{
-			bResult = ComplexMatches(MatchTypeOne, Other, MatchTypeTwo);
-		}
-#if CHECK_TAG_OPTIMIZATIONS
-		check(bResult == MatchesOriginal(MatchTypeOne, Other, MatchTypeTwo));
-#endif
-		return bResult;
-	}
+	int32 MatchesTagDepth(const FGameplayTag& TagToCheck) const;
+
 	/**
-	 * Check to see if two FGameplayTags match
+	 * Checks if this tag matches ANY of the tags in the specified container, also checks against our parent tags
+	 * "A.1".MatchesAny({"A","B"}) will return True, "A".MatchesAny({"A.1","B"}) will return False
+	 * If ContainerToCheck is empty/invalid it will always return False
 	 *
-	 * @param MatchTypeOne	How we compare this tag, Explicitly or a match with any parents as well
-	 * @param Other			The second tag to compare against
-	 * @param MatchTypeTwo	How we compare Other tag, Explicitly or a match with any parents as well
-	 * 
-	 * @return True if there is a match according to the specified match types; false if not
+	 * @return True if this tag matches ANY of the tags of in ContainerToCheck
 	 */
-	bool ComplexMatches(TEnumAsByte<EGameplayTagMatchType::Type> MatchTypeOne, const FGameplayTag& Other, TEnumAsByte<EGameplayTagMatchType::Type> MatchTypeTwo) const;
-#if CHECK_TAG_OPTIMIZATIONS
-	bool MatchesOriginal(TEnumAsByte<EGameplayTagMatchType::Type> MatchTypeOne, const FGameplayTag& Other, TEnumAsByte<EGameplayTagMatchType::Type> MatchTypeTwo) const;
-#endif
+	bool MatchesAny(const FGameplayTagContainer& ContainerToCheck) const;
+
+	/**
+	 * Checks if this tag matches ANY of the tags in the specified container, only allowing exact matches
+	 * "A.1".MatchesAny({"A","B"}) will return False
+	 * If ContainerToCheck is empty/invalid it will always return False
+	 *
+	 * @return True if this tag matches ANY of the tags of in ContainerToCheck exactly
+	 */
+	bool MatchesAnyExact(const FGameplayTagContainer& ContainerToCheck) const;
+
 	/** Returns whether the tag is valid or not; Invalid tags are set to NAME_None and do not exist in the game-specific global dictionary */
 	FORCEINLINE bool IsValid() const
 	{
 		return (TagName != NAME_None);
 	}
 
-	/** Used so we can have a TMap of this struct*/
-	friend uint32 GetTypeHash(const FGameplayTag& Tag)
+	/** Returns reference to a GameplayTagContainer containing only this tag */
+	const FGameplayTagContainer& GetSingleTagContainer() const;
+
+	/** Returns direct parent GameplayTag of this GameplayTag, calling on x.y will return x */
+	FGameplayTag RequestDirectParent() const;
+
+	/** Returns a new container explicitly containing the tags of this tag */
+	FGameplayTagContainer GetGameplayTagParents() const;
+
+	/** Used so we can have a TMap of this struct */
+	FORCEINLINE friend uint32 GetTypeHash(const FGameplayTag& Tag)
 	{
 		return ::GetTypeHash(Tag.TagName);
 	}
 
 	/** Displays gameplay tag as a string for blueprint graph usage */
-	FString ToString() const
+	FORCEINLINE FString ToString() const
 	{
 		return TagName.ToString();
 	}
 
 	/** Get the tag represented as a name */
-	FName GetTagName() const
+	FORCEINLINE FName GetTagName() const
 	{
 		return TagName;
 	}
@@ -128,9 +164,6 @@
 		return Ar;
 	}
 
-	/** Returns direct parent GameplayTag of this GameplayTag */
-	FGameplayTag RequestDirectParent() const;
-
 	/** Overridden for fast serialize */
 	bool NetSerialize(FArchive& Ar, class UPackageMap* Map, bool& bOutSuccess);
 
@@ -144,6 +177,50 @@
 	/** Handles importing tag strings without (TagName=) in it */
 	bool ImportTextItem(const TCHAR*& Buffer, int32 PortFlags, UObject* Parent, FOutputDevice* ErrorText);
 
+	/** An empty Gameplay Tag */
+	static const FGameplayTag EmptyTag;
+
+	// DEPRECATED
+
+PRAGMA_DISABLE_DEPRECATION_WARNINGS
+
+	/**
+	 * Check to see if two FGameplayTags match with explicit match types
+	 *
+	 * @param MatchTypeOne	How we compare this tag, Explicitly or a match with any parents as well
+	 * @param Other			The second tag to compare against
+	 * @param MatchTypeTwo	How we compare Other tag, Explicitly or a match with any parents as well
+	 * 
+	 * @return True if there is a match according to the specified match types; false if not
+	 */
+	DEPRECATED(4.15, "Deprecated in favor of MatchesTag")
+	FORCEINLINE_DEBUGGABLE bool Matches(TEnumAsByte<EGameplayTagMatchType::Type> MatchTypeOne, const FGameplayTag& Other, TEnumAsByte<EGameplayTagMatchType::Type> MatchTypeTwo) const
+	{
+		bool bResult = false;
+		if (MatchTypeOne == EGameplayTagMatchType::Explicit && MatchTypeTwo == EGameplayTagMatchType::Explicit)
+		{
+			bResult = TagName == Other.TagName;
+		}
+		else
+		{
+			bResult = ComplexMatches(MatchTypeOne, Other, MatchTypeTwo);
+		}
+		return bResult;
+	}
+	/**
+	 * Check to see if two FGameplayTags match
+	 *
+	 * @param MatchTypeOne	How we compare this tag, Explicitly or a match with any parents as well
+	 * @param Other			The second tag to compare against
+	 * @param MatchTypeTwo	How we compare Other tag, Explicitly or a match with any parents as well
+	 * 
+	 * @return True if there is a match according to the specified match types; false if not
+	 */
+	DEPRECATED(4.15, "Deprecated in favor of MatchesTag")
+	bool ComplexMatches(TEnumAsByte<EGameplayTagMatchType::Type> MatchTypeOne, const FGameplayTag& Other, TEnumAsByte<EGameplayTagMatchType::Type> MatchTypeTwo) const;
+
+PRAGMA_ENABLE_DEPRECATION_WARNINGS
+
 private:
 
 	/** Intentionally private so only the tag manager can use */
@@ -155,6 +232,7 @@
 
 	friend class UGameplayTagsManager;
 	friend struct FGameplayTagContainer;
+	friend struct FGameplayTagNode;
 };
 
 template<>
@@ -169,8 +247,8 @@
 	};
 };
 
-/** Simple struct for a gameplay tag container */
-USTRUCT(BlueprintType)
+/** A Tag Container holds a collection of FGameplayTags, tags are included explicitly by adding them, and implicitly from adding child tags */
+USTRUCT(BlueprintType, meta = (HasNativeMake = "GameplayTags.BlueprintGameplayTagLibrary.MakeGameplayTagContainerFromArray", HasNativeBreak = "GameplayTags.BlueprintGameplayTagLibrary.BreakGameplayTagContainer"))
 struct GAMEPLAYTAGS_API FGameplayTagContainer
 {
 	GENERATED_USTRUCT_BODY()
@@ -179,88 +257,208 @@
 	FGameplayTagContainer()
 	{
 	}
+
 	FGameplayTagContainer(FGameplayTagContainer const& Other)
 	{
 		*this = Other;
 	}
 
-	FGameplayTagContainer(const FGameplayTag& Tag)
+	/** Explicit to prevent people from accidentally using the wrong type of operation */
+	explicit FGameplayTagContainer(const FGameplayTag& Tag)
 	{
 		AddTag(Tag);
 	}
 
 	FGameplayTagContainer(FGameplayTagContainer&& Other)
 		: GameplayTags(MoveTemp(Other.GameplayTags))
+		, ParentTags(MoveTemp(Other.ParentTags))
 	{
 
 	}
+
 	~FGameplayTagContainer()
 	{
 	}
 
+	template<class AllocatorType>
+	static FGameplayTagContainer CreateFromArray(const TArray<FGameplayTag, AllocatorType>& SourceTags)
+	{
+		FGameplayTagContainer Container;
+		Container.GameplayTags.Append(SourceTags);
+		Container.FillParentTags();
+		return Container;
+	}
+
 	/** Assignment/Equality operators */
 	FGameplayTagContainer& operator=(FGameplayTagContainer const& Other);
 	FGameplayTagContainer& operator=(FGameplayTagContainer&& Other);
 	bool operator==(FGameplayTagContainer const& Other) const;
 	bool operator!=(FGameplayTagContainer const& Other) const;
 
-	template<class AllocatorType>
-	static FGameplayTagContainer CreateFromArray(TArray<FGameplayTag, AllocatorType>& SourceTags)
+	/**
+	 * Determine if TagToCheck is present in this container, also checking against parent tags
+	 * {"A.1"}.HasTag("A") will return True, {"A"}.HasTag("A.1") will return False
+	 * If TagToCheck is not Valid it will always return False
+	 * 
+	 * @return True if TagToCheck is in this container, false if it is not
+	 */
+	FORCEINLINE_DEBUGGABLE bool HasTag(const FGameplayTag& TagToCheck) const
 	{
-		FGameplayTagContainer Container;
-		Container.GameplayTags.Append(SourceTags);
-		return Container;
+		if (!TagToCheck.IsValid())
+		{
+			return false;
+		}
+		// Check explicit and parent tag list 
+		return GameplayTags.Contains(TagToCheck) || ParentTags.Contains(TagToCheck);
 	}
 
-	/**  Returns a new container containing all of the tags of this container, as well as all of their parent tags */
+	/**
+	 * Determine if TagToCheck is explicitly present in this container, only allowing exact matches
+	 * {"A.1"}.HasTagExact("A") will return False
+	 * If TagToCheck is not Valid it will always return False
+	 * 
+	 * @return True if TagToCheck is in this container, false if it is not
+	 */
+	FORCEINLINE_DEBUGGABLE bool HasTagExact(const FGameplayTag& TagToCheck) const
+	{
+		if (!TagToCheck.IsValid())
+		{
+			return false;
+		}
+		// Only check check explicit tag list
+		return GameplayTags.Contains(TagToCheck);
+	}
+
+	/**
+	 * Checks if this container contains ANY of the tags in the specified container, also checks against parent tags
+	 * {"A.1"}.HasAny({"A","B"}) will return True, {"A"}.HasAny({"A.1","B"}) will return False
+	 * If ContainerToCheck is empty/invalid it will always return False
+	 *
+	 * @return True if this container has ANY of the tags of in ContainerToCheck
+	 */
+	FORCEINLINE_DEBUGGABLE bool HasAny(const FGameplayTagContainer& ContainerToCheck) const
+	{
+		if (ContainerToCheck.IsEmpty())
+		{
+			return false;
+		}
+		for (const FGameplayTag& OtherTag : ContainerToCheck.GameplayTags)
+		{
+			if (GameplayTags.Contains(OtherTag) || ParentTags.Contains(OtherTag))
+			{
+				return true;
+			}
+		}
+		return false;
+	}
+
+	/**
+	 * Checks if this container contains ANY of the tags in the specified container, only allowing exact matches
+	 * {"A.1"}.HasAny({"A","B"}) will return False
+	 * If ContainerToCheck is empty/invalid it will always return False
+	 *
+	 * @return True if this container has ANY of the tags of in ContainerToCheck
+	 */
+	FORCEINLINE_DEBUGGABLE bool HasAnyExact(const FGameplayTagContainer& ContainerToCheck) const
+	{
+		if (ContainerToCheck.IsEmpty())
+		{
+			return false;
+		}
+		for (const FGameplayTag& OtherTag : ContainerToCheck.GameplayTags)
+		{
+			if (GameplayTags.Contains(OtherTag))
+			{
+				return true;
+			}
+		}
+		return false;
+	}
+
+	/**
+	 * Checks if this container contains ALL of the tags in the specified container, also checks against parent tags
+	 * {"A.1","B.1"}.HasAll({"A","B"}) will return True, {"A","B"}.HasAll({"A.1","B.1"}) will return False
+	 * If ContainerToCheck is empty/invalid it will always return True, because there were no failed checks
+	 *
+	 * @return True if this container has ALL of the tags of in ContainerToCheck, including if ContainerToCheck is empty
+	 */
+	FORCEINLINE_DEBUGGABLE bool HasAll(const FGameplayTagContainer& ContainerToCheck) const
+	{
+		if (ContainerToCheck.IsEmpty())
+		{
+			return true;
+		}
+		for (const FGameplayTag& OtherTag : ContainerToCheck.GameplayTags)
+		{
+			if (!GameplayTags.Contains(OtherTag) && !ParentTags.Contains(OtherTag))
+			{
+				return false;
+			}
+		}
+		return true;
+	}
+
+	/**
+	 * Checks if this container contains ALL of the tags in the specified container, only allowing exact matches
+	 * {"A.1","B.1"}.HasAll({"A","B"}) will return False
+	 * If ContainerToCheck is empty/invalid it will always return True, because there were no failed checks
+	 *
+	 * @return True if this container has ALL of the tags of in ContainerToCheck, including if ContainerToCheck is empty
+	 */
+	FORCEINLINE_DEBUGGABLE bool HasAllExact(const FGameplayTagContainer& ContainerToCheck) const
+	{
+		if (ContainerToCheck.IsEmpty())
+		{
+			return true;
+		}
+		for (const FGameplayTag& OtherTag : ContainerToCheck.GameplayTags)
+		{
+			if (!GameplayTags.Contains(OtherTag))
+			{
+				return false;
+			}
+		}
+		return true;
+	}
+
+	/** Returns the number of explicitly added tags */
+	FORCEINLINE int32 Num() const
+	{
+		return GameplayTags.Num();
+	}
+
+	/** Returns whether the container has any valid tags */
+	FORCEINLINE bool IsValid() const
+	{
+		return GameplayTags.Num() > 0;
+	}
+
+	/** Returns true if container is empty */
+	FORCEINLINE bool IsEmpty() const
+	{
+		return GameplayTags.Num() == 0;
+	}
+
+	/** Returns a new container explicitly containing the tags of this container and all of their parent tags */
 	FGameplayTagContainer GetGameplayTagParents() const;
 
 	/**
-	 * Returns a filtered version of this container, as if the container were filtered by matches from the parameter container
+	 * Returns a filtered version of this container, returns all tags that match against any of the tags in OtherContainer, expanding parents
 	 *
 	 * @param OtherContainer		The Container to filter against
-	 * @param TagMatchType			Type of match to use for the tags in this container
-	 * @param OtherTagMatchType		Type of match to use for the tags in the OtherContainer param
 	 *
 	 * @return A FGameplayTagContainer containing the filtered tags
 	 */
-	FGameplayTagContainer Filter(const FGameplayTagContainer& OtherContainer, TEnumAsByte<EGameplayTagMatchType::Type> TagMatchType, TEnumAsByte<EGameplayTagMatchType::Type> OtherTagMatchType) const;
+	FGameplayTagContainer Filter(const FGameplayTagContainer& OtherContainer) const;
 
 	/**
-	 * Checks if this container matches ANY of the tags in the specified container. Performs matching by expanding this container out
-	 * to include its parent tags.
+	 * Returns a filtered version of this container, returns all tags that match exactly one in OtherContainer
 	 *
-	 * @param Other					Container we are checking against
-	 * @param bCountEmptyAsMatch	If true, the parameter tag container will count as matching even if it's empty
+	 * @param OtherContainer		The Container to filter against
 	 *
-	 * @return True if this container has ANY the tags of the passed in container
+	 * @return A FGameplayTagContainer containing the filtered tags
 	 */
-	FORCEINLINE_DEBUGGABLE bool MatchesAny(const FGameplayTagContainer& Other, bool bCountEmptyAsMatch) const
-	{
-		if (Other.Num() == 0)
-		{
-			return bCountEmptyAsMatch;
-		}
-		return DoesTagContainerMatch(Other, EGameplayTagMatchType::IncludeParentTags, EGameplayTagMatchType::Explicit, EGameplayContainerMatchType::Any);
-	}
-
-	/**
-	* Checks if this container matches ALL of the tags in the specified container. Performs matching by expanding this container out to
-	* include its parent tags.
-	*
-	* @param Other				Container we are checking against
-	* @param bCountEmptyAsMatch	If true, the parameter tag container will count as matching even if it's empty
-	* 
-	* @return True if this container has ALL the tags of the passed in container
-	*/
-	FORCEINLINE_DEBUGGABLE bool MatchesAll(const FGameplayTagContainer& Other, bool bCountEmptyAsMatch) const
-	{
-		if (Other.Num() == 0)
-		{
-			return bCountEmptyAsMatch;
-		}
-		return DoesTagContainerMatch(Other, EGameplayTagMatchType::IncludeParentTags, EGameplayTagMatchType::Explicit, EGameplayContainerMatchType::All);
-	}
+	FGameplayTagContainer FilterExact(const FGameplayTagContainer& OtherContainer) const;
 
 	/** 
 	 * Checks if this container matches the given query.
@@ -271,59 +469,5 @@
 	 */
 	bool MatchesQuery(const struct FGameplayTagQuery& Query) const;
 
-
-	/**
-	 * Determine if the container has the specified tag. This forces an explicit match. 
-	 * This function exists for convenience and brevity. We do not wish to use default values for ::HasTag match type parameters, to avoid confusion on what the default behavior is. (E.g., we want people to think and use the right match type).
-	 * 
-	 * @param TagToCheck			Tag to check if it is present in the container
-	 * 
-	 * @return True if the tag is in the container, false if it is not
-	 */
-	FORCEINLINE_DEBUGGABLE bool HasTagExplicit(FGameplayTag const& TagToCheck) const
-	{
-		return HasTag(TagToCheck, EGameplayTagMatchType::Explicit, EGameplayTagMatchType::Explicit);
-	}
-
-	/**
-	 * Determine if the container has the specified tag
-	 * 
-	 * @param TagToCheck			Tag to check if it is present in the container
-	 * @param TagMatchType			Type of match to use for the tags in this container
-	 * @param TagToCheckMatchType	Type of match to use for the TagToCheck Param
-	 * 
-	 * @return True if the tag is in the container, false if it is not
-	 */
-	FORCEINLINE_DEBUGGABLE bool HasTag(FGameplayTag const& TagToCheck, TEnumAsByte<EGameplayTagMatchType::Type> TagMatchType=EGameplayTagMatchType::Explicit, TEnumAsByte<EGameplayTagMatchType::Type> TagToCheckMatchType=EGameplayTagMatchType::Explicit) const
-	{
-		SCOPE_CYCLE_COUNTER(STAT_FGameplayTagContainer_HasTag);
-		bool bResult;
-		if (TagMatchType == EGameplayTagMatchType::Explicit && TagToCheckMatchType == EGameplayTagMatchType::Explicit)
-		{
-			bResult = GameplayTags.Contains(TagToCheck);
-		}
-		else
-		{
-			bResult = ComplexHasTag(TagToCheck, TagMatchType, TagToCheckMatchType);
-		}
-#if CHECK_TAG_OPTIMIZATIONS
-		check(bResult == HasTagOriginal(TagToCheck, TagMatchType, TagToCheckMatchType));
-#endif
-		return bResult;
-	}
-	/**
-	 * Determine if the container has the specified tag
-	 * 
-	 * @param TagToCheck			Tag to check if it is present in the container
-	 * @param TagMatchType			Type of match to use for the tags in this container
-	 * @param TagToCheckMatchType	Type of match to use for the TagToCheck Param
-	 * 
-	 * @return True if the tag is in the container, false if it is not
-	 */
-	bool ComplexHasTag(FGameplayTag const& TagToCheck, TEnumAsByte<EGameplayTagMatchType::Type> TagMatchType, TEnumAsByte<EGameplayTagMatchType::Type> TagToCheckMatchType) const;
-#if CHECK_TAG_OPTIMIZATIONS
-	bool HasTagOriginal(FGameplayTag const& TagToCheck, TEnumAsByte<EGameplayTagMatchType::Type> TagMatchType, TEnumAsByte<EGameplayTagMatchType::Type> TagToCheckMatchType) const;
-#endif
-
 	/** 
 	 * Adds all the tags from one container to this container 
@@ -335,7 +479,7 @@
 	/** 
 	 * Adds all the tags that match between the two specified containers to this container 
 	 *
-	 * @param OtherA TagContainer that has the matching tags you want to add to this container 
+	 * @param OtherA TagContainer that has the matching tags you want to add to this container, these tags have their parents expanded
 	 * @param OtherB TagContainer used to check for matching tags
 	 */
 	void AppendMatchingTags(FGameplayTagContainer const& OtherA, FGameplayTagContainer const& OtherB);
@@ -357,5 +501,14 @@
 	void AddTagFast(const FGameplayTag& TagToAdd);
 
 	/**
+	 * Adds a tag to the container and removes any direct parents, wont add if child already exists
+	 *
+	 * @param Tag			The tag to try and add to this container
+	 * 
+	 * @return True if tag was added
+	 */
+	bool AddLeafTag(const FGameplayTag& TagToAdd);
+
+	/**
 	 * Tag to remove from the container
 	 * 
@@ -364,42 +517,25 @@
 	bool RemoveTag(FGameplayTag TagToRemove);
 
 	/**
-	* Removes all tags in TagsToRemove from this container
-	*
-	* @param TagsToRemove	Tags to remove from the container
-	*/
+	 * Removes all tags in TagsToRemove from this container
+	 *
+	 * @param TagsToRemove	Tags to remove from the container
+	 */
 	void RemoveTags(FGameplayTagContainer TagsToRemove);
 
-	/** Remove all tags from the container */
-	void RemoveAllTags(int32 Slack=0);
-
-	void RemoveAllTagsKeepSlack()
-	{
-		RemoveAllTags(GameplayTags.Num());
-	}
-
-	void Reset()
-	{
-		GameplayTags.Reset();
-	}
-
-	/**
-	 * Serialize the tag container
-	 *
-	 * @param Ar	Archive to serialize to
-	 *
-	 * @return True if the container was serialized
-	 */
+	/** Remove all tags from the container. Will maintain slack by default */
+	void Reset(int32 Slack = 0);
+	
+	/** Serialize the tag container */
 	bool Serialize(FArchive& Ar);
 
-	/**
-	 * Returns the Tag Count
-	 *
-	 * @return The number of tags
-	 */
-	int32 Num() const;
+	/** Efficient network serialize, takes advantage of the dictionary */
+	bool NetSerialize(FArchive& Ar, class UPackageMap* Map, bool& bOutSuccess);
 
-	/** Returns human readable Tag list */
+	/** Handles fixup after importing from text */
+	bool ImportTextItem(const TCHAR*& Buffer, int32 PortFlags, UObject* Parent, FOutputDevice* ErrorText);
+
+	/** Returns string version of container in ImportText format */
 	FString ToString() const;
 
 	/** Returns abbreviated human readable Tag list without parens or property names */
@@ -408,6 +544,12 @@
 	// Returns human readable description of what match is being looked for on the readable tag list.
 	FText ToMatchingText(EGameplayContainerMatchType MatchType, bool bInvertCondition) const;
 
+	/** Gets the explicit list of gameplay tags */
+	void GetGameplayTagArray(TArray<FGameplayTag>& InOutGameplayTags) const
+	{
+		InOutGameplayTags = GameplayTags;
+	}
+
 	/** Creates a const iterator for the contents of this array */
 	TArray<FGameplayTag>::TConstIterator CreateConstIterator() const
 	{
@@ -440,45 +582,186 @@
 
 	/** An empty Gameplay Tag Container */
 	static const FGameplayTagContainer EmptyContainer;
-		
-	bool NetSerialize(FArchive& Ar, class UPackageMap* Map, bool& bOutSuccess);
+
+	// DEPRECATED FUNCTIONALITY
+
+PRAGMA_DISABLE_DEPRECATION_WARNINGS
+
+	DEPRECATED(4.15, "Deprecated in favor of Reset")
+		void RemoveAllTags(int32 Slack = 0)
+	{
+		Reset(Slack);
+	}
+
+	DEPRECATED(4.15, "Deprecated in favor of Reset")
+		void RemoveAllTagsKeepSlack()
+	{
+		Reset();
+	}
 
 	/**
-	* Returns true if the tags in this container match the tags in OtherContainer for the specified matching types.
-	*
-	* @param OtherContainer		The Container to filter against
-	* @param TagMatchType			Type of match to use for the tags in this container
-	* @param OtherTagMatchType		Type of match to use for the tags in the OtherContainer param
-	* @param ContainerMatchType	Type of match to use for filtering
-	*
-	* @return Returns true if ContainerMatchType is Any and any of the tags in OtherContainer match the tags in this or ContainerMatchType is All and all of the tags in OtherContainer match at least one tag in this. Returns false otherwise.
-	*/
+	 * Determine if the container has the specified tag. This forces an explicit match. 
+	 * This function exists for convenience and brevity. We do not wish to use default values for ::HasTag match type parameters, to avoid confusion on what the default behavior is. (E.g., we want people to think and use the right match type).
+	 * 
+	 * @param TagToCheck			Tag to check if it is present in the container
+	 * 
+	 * @return True if the tag is in the container, false if it is not
+	 */
+	DEPRECATED(4.15, "Deprecated in favor of HasTagExact")
+	FORCEINLINE_DEBUGGABLE bool HasTagExplicit(FGameplayTag const& TagToCheck) const
+	{
+		return HasTag(TagToCheck, EGameplayTagMatchType::Explicit, EGameplayTagMatchType::Explicit);
+	}
+
+	/**
+	 * Determine if the container has the specified tag
+	 * 
+	 * @param TagToCheck			Tag to check if it is present in the container
+	 * @param TagMatchType			Type of match to use for the tags in this container
+	 * @param TagToCheckMatchType	Type of match to use for the TagToCheck Param
+	 * 
+	 * @return True if the tag is in the container, false if it is not
+	 */
+	DEPRECATED(4.15, "Deprecated in favor of HasTag with no parameters")
+	FORCEINLINE_DEBUGGABLE bool HasTag(FGameplayTag const& TagToCheck, TEnumAsByte<EGameplayTagMatchType::Type> TagMatchType, TEnumAsByte<EGameplayTagMatchType::Type> TagToCheckMatchType) const
+	{
+		SCOPE_CYCLE_COUNTER(STAT_FGameplayTagContainer_HasTag);
+		if (!TagToCheck.IsValid())
+		{
+			return false;
+		}
+
+		return HasTagFast(TagToCheck, TagMatchType, TagToCheckMatchType);
+	}
+
+	/** Version of above that is called from conditions where you know tag is valid */
+	FORCEINLINE_DEBUGGABLE bool HasTagFast(FGameplayTag const& TagToCheck, TEnumAsByte<EGameplayTagMatchType::Type> TagMatchType, TEnumAsByte<EGameplayTagMatchType::Type> TagToCheckMatchType) const
+	{
+		bool bResult;
+		if (TagToCheckMatchType == EGameplayTagMatchType::Explicit)
+		{
+			// Always check explicit
+			bResult = GameplayTags.Contains(TagToCheck);
+
+			if (!bResult && TagMatchType == EGameplayTagMatchType::IncludeParentTags)
+			{
+				// Check parent tags as well
+				bResult = ParentTags.Contains(TagToCheck);
+			}
+		}
+		else
+		{
+			bResult = ComplexHasTag(TagToCheck, TagMatchType, TagToCheckMatchType);
+		}
+		return bResult;
+	}
+
+	/**
+	 * Determine if the container has the specified tag
+	 * 
+	 * @param TagToCheck			Tag to check if it is present in the container
+	 * @param TagMatchType			Type of match to use for the tags in this container
+	 * @param TagToCheckMatchType	Type of match to use for the TagToCheck Param
+	 * 
+	 * @return True if the tag is in the container, false if it is not
+	 */
+	bool ComplexHasTag(FGameplayTag const& TagToCheck, TEnumAsByte<EGameplayTagMatchType::Type> TagMatchType, TEnumAsByte<EGameplayTagMatchType::Type> TagToCheckMatchType) const;
+
+	/**
+	 * Checks if this container matches ANY of the tags in the specified container. Performs matching by expanding this container out
+	 * to include its parent tags.
+	 *
+	 * @param Other					Container we are checking against
+	 * @param bCountEmptyAsMatch	If true, the parameter tag container will count as matching even if it's empty
+	 *
+	 * @return True if this container has ANY the tags of the passed in container
+	 */
+	DEPRECATED(4.15, "Deprecated in favor of HasAny")
+	FORCEINLINE_DEBUGGABLE bool MatchesAny(const FGameplayTagContainer& Other, bool bCountEmptyAsMatch) const
+	{
+		if (Other.IsEmpty())
+		{
+			return bCountEmptyAsMatch;
+		}
+		return DoesTagContainerMatch(Other, EGameplayTagMatchType::IncludeParentTags, EGameplayTagMatchType::Explicit, EGameplayContainerMatchType::Any);
+	}
+
+	/**
+	 * Checks if this container matches ALL of the tags in the specified container. Performs matching by expanding this container out to
+	 * include its parent tags.
+	 *
+	 * @param Other				Container we are checking against
+	 * @param bCountEmptyAsMatch	If true, the parameter tag container will count as matching even if it's empty
+	 * 
+	 * @return True if this container has ALL the tags of the passed in container
+	 */
+	DEPRECATED(4.15, "Deprecated in favor of HasAll")
+	FORCEINLINE_DEBUGGABLE bool MatchesAll(const FGameplayTagContainer& Other, bool bCountEmptyAsMatch) const
+	{
+		if (Other.IsEmpty())
+		{
+			return bCountEmptyAsMatch;
+		}
+		return DoesTagContainerMatch(Other, EGameplayTagMatchType::IncludeParentTags, EGameplayTagMatchType::Explicit, EGameplayContainerMatchType::All);
+	}
+
+	/**
+	 * Returns true if the tags in this container match the tags in OtherContainer for the specified matching types.
+	 *
+	 * @param OtherContainer		The Container to filter against
+	 * @param TagMatchType			Type of match to use for the tags in this container
+	 * @param OtherTagMatchType		Type of match to use for the tags in the OtherContainer param
+	 * @param ContainerMatchType	Type of match to use for filtering
+	 *
+	 * @return Returns true if ContainerMatchType is Any and any of the tags in OtherContainer match the tags in this or ContainerMatchType is All and all of the tags in OtherContainer match at least one tag in this. Returns false otherwise.
+	 */
 	FORCEINLINE_DEBUGGABLE bool DoesTagContainerMatch(const FGameplayTagContainer& OtherContainer, TEnumAsByte<EGameplayTagMatchType::Type> TagMatchType, TEnumAsByte<EGameplayTagMatchType::Type> OtherTagMatchType, EGameplayContainerMatchType ContainerMatchType) const
 	{
 		SCOPE_CYCLE_COUNTER(STAT_FGameplayTagContainer_DoesTagContainerMatch);
 		bool bResult;
-		if (ContainerMatchType == EGameplayContainerMatchType::Any && TagMatchType == EGameplayTagMatchType::Explicit && OtherTagMatchType == EGameplayTagMatchType::Explicit)
+		if (OtherTagMatchType == EGameplayTagMatchType::Explicit)
 		{
-			bResult = false;
-			for (TArray<FGameplayTag>::TConstIterator OtherIt(OtherContainer.GameplayTags); OtherIt; ++OtherIt)
+			// Start true for all, start false for any
+			bResult = (ContainerMatchType == EGameplayContainerMatchType::All);
+			for (const FGameplayTag& OtherTag : OtherContainer.GameplayTags)
 			{
-				if (GameplayTags.Contains(*OtherIt))
+				if (HasTagFast(OtherTag, TagMatchType, OtherTagMatchType))
 				{
-					bResult = true;
+					if (ContainerMatchType == EGameplayContainerMatchType::Any)
+					{
+						bResult = true;
+						break;
+					}
+				}
+				else if (ContainerMatchType == EGameplayContainerMatchType::All)
+				{
+					bResult = false;
 					break;
 				}
 			}			
 		}
 		else
 		{
-			bResult = DoesTagContainerMatchComplex(OtherContainer, TagMatchType, OtherTagMatchType, ContainerMatchType);
+			FGameplayTagContainer OtherExpanded = OtherContainer.GetGameplayTagParents();
+			return DoesTagContainerMatch(OtherExpanded, TagMatchType, EGameplayTagMatchType::Explicit, ContainerMatchType);
 		}
-#if CHECK_TAG_OPTIMIZATIONS
-		// the complex version is the original
-		check(bResult == DoesTagContainerMatchComplex(OtherContainer, TagMatchType, OtherTagMatchType, ContainerMatchType));
-#endif
 		return bResult;
 	}
+
+	/**
+	 * Returns a filtered version of this container, as if the container were filtered by matches from the parameter container
+	 *
+	 * @param OtherContainer		The Container to filter against
+	 * @param TagMatchType			Type of match to use for the tags in this container
+	 * @param OtherTagMatchType		Type of match to use for the tags in the OtherContainer param
+	 *
+	 * @return A FGameplayTagContainer containing the filtered tags
+	 */
+	DEPRECATED(4.15, "Deprecated in favor of HasAll")
+	FGameplayTagContainer Filter(const FGameplayTagContainer& OtherContainer, TEnumAsByte<EGameplayTagMatchType::Type> TagMatchType, TEnumAsByte<EGameplayTagMatchType::Type> OtherTagMatchType) const;
+
+PRAGMA_ENABLE_DEPRECATION_WARNINGS
+
 protected:
 
 	/**
@@ -493,10 +776,6 @@
 	*/
 	bool DoesTagContainerMatchComplex(const FGameplayTagContainer& OtherContainer, TEnumAsByte<EGameplayTagMatchType::Type> TagMatchType, TEnumAsByte<EGameplayTagMatchType::Type> OtherTagMatchType, EGameplayContainerMatchType ContainerMatchType) const;
 
-	/** Array of gameplay tags */
-	UPROPERTY(BlueprintReadWrite, Category=GameplayTags)
-	TArray<FGameplayTag> GameplayTags;
-
 	/**
 	 * If a Tag with the specified tag name explicitly exists, it will remove that tag and return true.  Otherwise, it 
 	   returns false.  It does NOT check the TagName for validity (i.e. the tag could be obsolete and so not exist in
@@ -506,17 +785,27 @@
 	 */
 	bool RemoveTagByExplicitName(const FName& TagName);
 
-	// Allow the redirection helper class access to RemoveTagByExplicitName.  It can then (through friendship) allow
-	// access to others without exposing everything the Container has privately to everyone.
-	friend class FGameplayTagRedirectHelper;
+	/** Adds parent tags for a single tag */
+	void AddParentsForTag(const FGameplayTag& Tag);
+
+	/** Fills in ParentTags from GameplayTags */
+	void FillParentTags();
+
+	/** Array of gameplay tags */
+	UPROPERTY(BlueprintReadWrite, Category=GameplayTags) // Change to VisibleAnywhere after fixing up games
+	TArray<FGameplayTag> GameplayTags;
+
+	/** Array of expanded parent tags, in addition to GameplayTags. Used to accelerate parent searches. May contain duplicates in some cases */
+	UPROPERTY(Transient)
+	TArray<FGameplayTag> ParentTags;
+
+	friend class UGameplayTagsManager;
 	friend struct FGameplayTagQuery;
 	friend struct FGameplayTagQueryExpression;
-
+	friend struct FGameplayTagNode;
+	friend struct FGameplayTag;
+	
 private:
 
-	/** Array of gameplay tags */
-	UPROPERTY()
-	TArray<FName> Tags_DEPRECATED;
-
 	/**
 	 * DO NOT USE DIRECTLY
@@ -527,15 +816,14 @@
 	FORCEINLINE friend TArray<FGameplayTag>::TConstIterator end(const FGameplayTagContainer& Array) { return TArray<FGameplayTag>::TConstIterator(Array.GameplayTags, Array.GameplayTags.Num()); }
 };
 
-// This helper class exists to keep FGameplayTagContainers protected and private fields appropriately private while
-// exposing only necessary features to the GameplayTagsManager.
-class FGameplayTagRedirectHelper
+FORCEINLINE bool FGameplayTag::MatchesAnyExact(const FGameplayTagContainer& ContainerToCheck) const
 {
-private:
-	FORCEINLINE static bool RemoveTagByExplicitName(FGameplayTagContainer& Container, const FName& TagName) { return Container.RemoveTagByExplicitName(TagName); }
+	if (ContainerToCheck.IsEmpty())
+	{
+		return false;
+	}
+	return ContainerToCheck.GameplayTags.Contains(*this);
+}
 
-	friend class UGameplayTagsManager;
-};
-
 template<>
 struct TStructOpsTypeTraits<FGameplayTagContainer> : public TStructOpsTypeTraitsBase
@@ -545,6 +833,7 @@
 		WithSerializer = true,
 		WithIdenticalViaEquality = true,
 		WithNetSerializer = true,
+		WithImportTextItem = true,
 		WithCopy = true
 	};
 };
@@ -592,14 +881,14 @@
  *	Q.BuildQuery(
  *		FGameplayTagQueryExpression()
  * 		.AllTagsMatch()
- *		.AddTag(IGameplayTagsModule::RequestGameplayTag(FName(TEXT("Animal.Mammal.Dog.Corgi"))))
- *		.AddTag(IGameplayTagsModule::RequestGameplayTag(FName(TEXT("Plant.Tree.Spruce"))))
+ *		.AddTag(FGameplayTag::RequestGameplayTag(FName(TEXT("Animal.Mammal.Dog.Corgi"))))
+ *		.AddTag(FGameplayTag::RequestGameplayTag(FName(TEXT("Plant.Tree.Spruce"))))
  *		);
  * 
  * Queries are internally represented as a byte stream that is memory-efficient and can be evaluated quickly at runtime.
  * Note: these have an extensive details and graph pin customization for editing, so there is no need to expose the internals to Blueprints.
  */
-USTRUCT(BlueprintType, meta=(HasNativeMake="GameplayTags.UBlueprintGameplayTagLibrary.MakeGameplayTagQuery"))
+USTRUCT(BlueprintType, meta=(HasNativeMake="GameplayTags.BlueprintGameplayTagLibrary.MakeGameplayTagQuery"))
 struct GAMEPLAYTAGS_API FGameplayTagQuery
 {
 	GENERATED_BODY();
diff --git a/Engine/Source/Runtime/GameplayTags/Classes/GameplayTagsManager.h b/Engine/Source/Runtime/GameplayTags/Classes/GameplayTagsManager.h
--- a/Engine/Source/Runtime/GameplayTags/Classes/GameplayTagsManager.h
+++ b/Engine/Source/Runtime/GameplayTags/Classes/GameplayTagsManager.h
@@ -13,25 +13,74 @@ struct FGameplayTagTableRow : public FTableRowBase
 	GENERATED_USTRUCT_BODY()
 
 	/** Tag specified in the table */
-	UPROPERTY(EditAnywhere, BlueprintReadWrite, Category=GameplayTag)
-	FString Tag;
-
-	/** Text that describes this category - not all tags have categories */
-	UPROPERTY(EditAnywhere, BlueprintReadWrite, Category=GameplayTag)
-	FText CategoryText;
+	UPROPERTY(EditAnywhere, BlueprintReadOnly, Category=GameplayTag)
+	FName Tag;
 
 	/** Developer comment clarifying the usage of a particular tag, not user facing */
 	UPROPERTY(EditAnywhere, BlueprintReadOnly, Category=GameplayTag)
 	FString DevComment;
 
 	/** Constructors */
-	FGameplayTagTableRow() { CategoryText = NSLOCTEXT("TAGCATEGORY", "TAGCATEGORYTESTING", "Category and Category"); }
-	FGameplayTagTableRow(FGameplayTagTableRow const& Other);
+	FGameplayTagTableRow() {}
+	FGameplayTagTableRow(FName InTag, const FString& InDevComment = TEXT("")) : Tag(InTag), DevComment(InDevComment) {}
+	GAMEPLAYTAGS_API FGameplayTagTableRow(FGameplayTagTableRow const& Other);
 
 	/** Assignment/Equality operators */
-	FGameplayTagTableRow& operator=(FGameplayTagTableRow const& Other);
-	bool operator==(FGameplayTagTableRow const& Other) const;
-	bool operator!=(FGameplayTagTableRow const& Other) const;
+	GAMEPLAYTAGS_API FGameplayTagTableRow& operator=(FGameplayTagTableRow const& Other);
+	GAMEPLAYTAGS_API bool operator==(FGameplayTagTableRow const& Other) const;
+	GAMEPLAYTAGS_API bool operator!=(FGameplayTagTableRow const& Other) const;
+	GAMEPLAYTAGS_API bool operator<(FGameplayTagTableRow const& Other) const;
+};
+
+UENUM()
+enum class EGameplayTagSourceType : uint8
+{
+	Native,				// Was added from C++ code
+	DefaultTagList,		// The default tag list in DefaultGameplayTags.ini
+	TagList,			// Another tag list from an ini in tags/*.ini
+	DataTable,			// From a DataTable
+	Invalid,			// Not a real source
+};
+
+/** Struct defining where gameplay tags are loaded/saved from. Mostly for the editor */
+USTRUCT()
+struct FGameplayTagSource
+{
+	GENERATED_USTRUCT_BODY()
+
+	/** Name of this source */
+	UPROPERTY()
+	FName SourceName;
+
+	/** Type of this source */
+	UPROPERTY()
+	EGameplayTagSourceType SourceType;
+
+	/** If this is bound to an ini object for saving, this is the one */
+	UPROPERTY()
+	class UGameplayTagsList* SourceTagList;
+
+	FGameplayTagSource() 
+		: SourceName(NAME_None), SourceType(EGameplayTagSourceType::Invalid), SourceTagList(nullptr) 
+	{
+	}
+
+	FGameplayTagSource(FName InSourceName, EGameplayTagSourceType InSourceType, UGameplayTagsList* InSourceTagList = nullptr) 
+		: SourceName(InSourceName), SourceType(InSourceType), SourceTagList(InSourceTagList)
+	{
+	}
+
+	static FName GetNativeName()
+	{
+		static FName NativeName = FName(TEXT("Native"));
+		return NativeName;
+	}
+
+	static FName GetDefaultName()
+	{
+		static FName DefaultName = FName(TEXT("DefaultGameplayTags.ini"));
+		return DefaultName;
+	}
 };
 
 /** Simple tree node for gameplay tags */
@@ -42,88 +91,86 @@ struct FGameplayTagNode
 	FGameplayTagNode(){};
 
 	/** Simple constructor */
-	FGameplayTagNode(FName InTag, TWeakPtr<FGameplayTagNode> InParentNode, FText InCategoryDescription, bool InWasAddedDirectlyFromDictionary);
+	FGameplayTagNode(FName InTag, TSharedPtr<FGameplayTagNode> InParentNode);
+
+	/** Returns a correctly constructed container with only this tag, useful for doing container queries */
+	FORCEINLINE const FGameplayTagContainer& GetSingleTagContainer() const { return CompleteTagWithParents; }
 
 	/**
 	 * Get the complete tag for the node, including all parent tags, delimited by periods
 	 * 
 	 * @return Complete tag for the node
 	 */
-	GAMEPLAYTAGS_API FName GetCompleteTag() const { return CompleteTag; }
+	FORCEINLINE const FGameplayTag& GetCompleteTag() const { return CompleteTagWithParents.Num() > 0 ? CompleteTagWithParents.GameplayTags[0] : FGameplayTag::EmptyTag; }
+	FORCEINLINE FName GetCompleteTagName() const { return GetCompleteTag().GetTagName(); }
+	FORCEINLINE FString GetCompleteTagString() const { return GetCompleteTag().ToString(); }
 
 	/**
 	 * Get the simple tag for the node (doesn't include any parent tags)
 	 * 
 	 * @return Simple tag for the node
 	 */
-	GAMEPLAYTAGS_API FName GetSimpleTag() const { return Tag; }
-
-	/**
-	 * Get the category description for the node
-	 * 
-	 * @return Translatable text that describes this tag category
-	 */
-	GAMEPLAYTAGS_API FText GetCategoryDescription() const { return CategoryDescription; }
+	FORCEINLINE FName GetSimpleTagName() const { return Tag; }
 
 	/**
 	 * Get the children nodes of this node
 	 * 
 	 * @return Reference to the array of the children nodes of this node
 	 */
-	GAMEPLAYTAGS_API TArray< TSharedPtr<FGameplayTagNode> >& GetChildTagNodes() { return ChildTags; }
+	FORCEINLINE TArray< TSharedPtr<FGameplayTagNode> >& GetChildTagNodes() { return ChildTags; }
 
 	/**
 	 * Get the children nodes of this node
 	 * 
 	 * @return Reference to the array of the children nodes of this node
 	 */
-	GAMEPLAYTAGS_API const TArray< TSharedPtr<FGameplayTagNode> >& GetChildTagNodes() const { return ChildTags; }
+	FORCEINLINE const TArray< TSharedPtr<FGameplayTagNode> >& GetChildTagNodes() const { return ChildTags; }
 
 	/**
 	 * Get the parent tag node of this node
 	 * 
 	 * @return The parent tag node of this node
 	 */
-	GAMEPLAYTAGS_API TWeakPtr<FGameplayTagNode> GetParentTagNode() const { return ParentNode; }
+	FORCEINLINE TSharedPtr<FGameplayTagNode> GetParentTagNode() const { return ParentNode; }
 
 	/**
 	* Get the net index of this node
 	*
 	* @return The net index of this node
 	*/
-	GAMEPLAYTAGS_API FGameplayTagNetIndex GetNetIndex() const { return NetIndex; }
+	FORCEINLINE FGameplayTagNetIndex GetNetIndex() const { return NetIndex; }
 
 	/** Reset the node of all of its values */
 	GAMEPLAYTAGS_API void ResetNode();
 
 private:
-	/** Tag for the node */
+	/** Raw name for this  */
 	FName Tag;
 
-	/** Complete tag for the node, including parent tags */
-	FName CompleteTag;
-
-	/** Category description of the for the node */
-	FText CategoryDescription;
-
-	/** Flag to distinguish implicit vs explicit tags */
-	bool WasAddedDirectlyFromDictionary;
+	/** This complete tag is at GameplayTags[0], with parents in ParentTags[] */
+	FGameplayTagContainer CompleteTagWithParents;
 
 	/** Child gameplay tag nodes */
 	TArray< TSharedPtr<FGameplayTagNode> > ChildTags;
 
 	/** Owner gameplay tag node, if any */
-	TWeakPtr<FGameplayTagNode> ParentNode;
-
-	FGameplayTagContainer Parents;
+	TSharedPtr<FGameplayTagNode> ParentNode;
 	
 	/** Net Index of this node */
 	FGameplayTagNetIndex NetIndex;
 
+#if WITH_EDITORONLY_DATA
+	/** Package or config file this tag came from. This is the first one added. If None, this is an implicitly added tag */
+	FName SourceName;
+
+	/** Comment for this tag */
+	FString DevComment;
+#endif 
+
 	friend class UGameplayTagsManager;
 };
 
-/** Holds global data loaded at startup, is in a singleton UObject so it works properly with hot reload */
+/** Holds data about the tag dictionary, is in a singleton UObject */
 UCLASS(config=Engine)
 class GAMEPLAYTAGS_API UGameplayTagsManager : public UObject
 {
@@ -132,84 +179,16 @@ class GAMEPLAYTAGS_API UGameplayTagsManager : public UObject
 	// Destructor
 	~UGameplayTagsManager();
 
-	/**
-	 * Helper function to insert a tag into a tag node array
-	 * 
-	 * @param Tag					Tag to insert
-	 * @param ParentNode			Parent node, if any, for the tag
-	 * @param NodeArray				Node array to insert the new node into, if necessary (if the tag already exists, no insertion will occur)
-	 * @param CategoryDescription	The description of this category
-	 * 
-	 * @return Index of the node of the tag
-	 */
-	int32 InsertTagIntoNodeArray(FName Tag, TWeakPtr<FGameplayTagNode> ParentNode, TArray< TSharedPtr<FGameplayTagNode> >& NodeArray, FText CategoryDescription, bool InWasAddedDirectlyFromDictionary);
-
-	/**
-	 * Get the best tag category description
-	 *
-	 * @param Tag				Tag that we will find all nodes for
-	 * @param OutDescription	The Description that will be returned
-	 *
-	 * @return The index that we got the description from
-	 */
-	int32 GetBestTagCategoryDescription(FString Tag, FText& OutDescription) const;
-
-	/** 
-	 * Gets all nodes that make up a tag, if any (e.g. weapons.ranged.pistol will return the nodes weapons, weapons.ranged, and weapons.ranged.pistol)
-	 * 
-	 * @param Tag			The . delimited tag we wish to get nodes for
-	 * @param OutTagArray	The array of tag nodes that were found
-	 */
-	void GetAllNodesForTag( const FString& Tag, TArray< TSharedPtr<FGameplayTagNode> >& OutTagArray ) const;
-
-	// this is here because the tag tree doesn't start with a node and recursion can't be done until the first node is found
-	void GetAllNodesForTag_Recurse(TArray<FString>& Tags, int32 CurrentTagDepth, TSharedPtr<FGameplayTagNode> CurrentTagNode, TArray< TSharedPtr<FGameplayTagNode> >& OutTagArray ) const;
-
-#if WITH_EDITOR
-	/** Gets a Filtered copy of the GameplayRootTags Array based on the comma delimited filter string passed in*/
-	void GetFilteredGameplayRootTags( const FString& InFilterString, TArray< TSharedPtr<FGameplayTagNode> >& OutTagArray );
-
-	/** 
-	 * Called via delegate when an object is re-imported in the editor
-	 * 
-	 * @param ImportFactory	Factory responsible for the re-import
-	 * @param InObject		Object that was re-imported
-	 */
-	void OnObjectReimported(class UFactory* ImportFactory, UObject* InObject);
-
-	TSharedPtr<FGameplayTagNode> FindTagNode(FName TagName) const;
-
-	static void AddNewGameplayTagToINI(FString NewTag);
-
-	bool IsDictionaryTag(FName TagName) const;
-
-#endif //WITH_EDITOR
-
-	struct FDeveloperTags 
+	/** Returns the global UGameplayTagsManager manager */
+	FORCEINLINE static UGameplayTagsManager& Get()
 	{
-		struct FDeveloperTagsItem
+		if (SingletonManager == nullptr)
 		{
-			FString IniName;
-			TArray<FString> Tags;
-		};
-
-		TArray<FDeveloperTagsItem> Items;
-	};
-
-	FDeveloperTags DeveloperTags;
-
-	/** 
-	 * Loads the tag tables
-	 * 
-	 * @param TagTableNames	The names of the tables to load
-	 */
-	void LoadGameplayTagTables(TArray<FString>& TagTableNames);
-
-	/** Helper function to construct the gameplay tag tree */
-	void ConstructGameplayTagTree();
+			InitializeManager();
+		}
 
-	/** Helper function to destroy the gameplay tag tree */
-	void DestroyGameplayTagTree();
+		return *SingletonManager;
+	}
 
 	/**
 	 * Gets the FGameplayTag that corresponds to the TagName
@@ -220,125 +199,96 @@ class GAMEPLAYTAGS_API UGameplayTagsManager : public UObject
 	 * 
 	 * @return Will return the corresponding FGameplayTag or an empty one if not found.
 	 */
-	UFUNCTION(BlueprintCallable, Category="GameplayTags")
 	FGameplayTag RequestGameplayTag(FName TagName, bool ErrorIfNotFound=true) const;
 
 	/**
-	 * Adds a tag to the container and removes any direct parents, wont add if child already exists
+	 * Gets a Tag Container containing the supplied tag and all of it's parents as explicit tags
 	 *
-	 * @param TagContainer	The tag container we want to add tag too
-	 * @param Tag			The tag to try and add to container
+	 * @param GameplayTag The Tag to use at the child most tag for this container
 	 * 
-	 * @return True if tag was added
+	 * @return A Tag Container with the supplied tag and all its parents added explicitly
 	 */
-	bool AddLeafTagToContainer(FGameplayTagContainer& TagContainer, const FGameplayTag& Tag);
+	FGameplayTagContainer RequestGameplayTagParents(const FGameplayTag& GameplayTag) const;
 
 	/**
-	 * Gets a Tag Container for the supplied tag and all its parents
+	 * Gets a Tag Container containing the all tags in the hierarchy that are children of this tag. Does not return the original tag
 	 *
-	 * @param GameplayTag The Tag to use at the child most tag for this container
+	 * @param GameplayTag					The Tag to use at the parent tag
 	 * 
-	 * @return A Tag Container with the supplied tag and all its parents
+	 * @return A Tag Container with the supplied tag and all its parents added explicitly
 	 */
-	FGameplayTagContainer RequestGameplayTagParents(const FGameplayTag& GameplayTag) const;
-
-	FGameplayTagContainer RequestGameplayTagChildren(const FGameplayTag& GameplayTag, bool OnlyIncludeDictionaryTags=false) const;
+	FGameplayTagContainer RequestGameplayTagChildren(const FGameplayTag& GameplayTag) const;
 
+	/** Returns direct parent GameplayTag of this GameplayTag, calling on x.y will return x */
 	FGameplayTag RequestGameplayTagDirectParent(const FGameplayTag& GameplayTag) const;
 
+	/**
+	 * Helper function to get the stored TagContainer containing only this tag, which has searchable ParentTags
+	 * @param GameplayTag		Tag to get single container of
+	 * @return					Pointer to container with this tag
+	 */
+	FORCEINLINE_DEBUGGABLE const FGameplayTagContainer* GetSingleTagContainer(const FGameplayTag& GameplayTag) const
+	{
+		const TSharedPtr<FGameplayTagNode>* TagNode = GameplayTagNodeMap.Find(GameplayTag);
+		if (TagNode)
+		{
+			return &((*TagNode)->GetSingleTagContainer());
+		}
+		return nullptr;
+	}
+
+	/** Loads the tag tables referenced in the GameplayTagSettings object */
+	void LoadGameplayTagTables();
+
+	/** Helper function to construct the gameplay tag tree */
+	void ConstructGameplayTagTree();
+
+	/** Helper function to destroy the gameplay tag tree */
+	void DestroyGameplayTagTree();
+
+	/** Splits a tag such as x.y.z into an array of names {x,y,z} */
 	void SplitGameplayTagFName(const FGameplayTag& Tag, TArray<FName>& OutNames);
 
 	/**
-	 * Checks if the tag is allowed to be created
+	 * Checks if the passed in name is in the tag dictionary and can be created
 	 *
 	 * @return True if valid
 	 */
 	bool ValidateTagCreation(FName TagName) const;
 
 	/**
-	 * Checks FGameplayTagNode and all FGameplayTagNode children to see if a FGameplayTagNode with the name exists
+	 * Checks node tree to see if a FGameplayTagNode with the name exists
 	 *
-	 * @param Node		The Node in the the to start searching from
 	 * @param TagName	The name of the tag node to search for
-	 * 
+	 *
 	 * @return A shared pointer to the FGameplayTagNode found, or NULL if not found.
 	 */
-	TSharedPtr<FGameplayTagNode> FindTagNode(TSharedPtr<FGameplayTagNode> Node, FName TagName) const;
+	TSharedPtr<FGameplayTagNode> FindTagNode(FName TagName) const;
+
+	/** Returns the tag source for a given tag source name, or null if not found */
+	const FGameplayTagSource* FindTagSource(FName TagSourceName) const;
+
+	/** Fills in an array with all tag sources of a specific type */
+	void FindTagSourcesWithType(EGameplayTagSourceType TagSourceType, TArray<const FGameplayTagSource*>& OutArray) const;
 
 	/**
-	 * Check to see if two FGameplayTags match
+	 * Check to see how closely two FGameplayTags match. Higher values indicate more matching terms in the tags.
 	 *
 	 * @param GameplayTagOne	The first tag to compare
-	 * @param MatchTypeOne		How we compare Tag one, Explicitly or a match with any parents as well
 	 * @param GameplayTagTwo	The second tag to compare
-	 * @param MatchTypeTwo		How we compare Tag two, Explicitly or a match with any parents as well
-	 * 
-	 * @return true if there is a match
-	 */
-	FORCEINLINE_DEBUGGABLE bool GameplayTagsMatch(const FGameplayTag& GameplayTagOne, TEnumAsByte<EGameplayTagMatchType::Type> MatchTypeOne, const FGameplayTag& GameplayTagTwo, TEnumAsByte<EGameplayTagMatchType::Type> MatchTypeTwo) const
-	{
-		SCOPE_CYCLE_COUNTER(STAT_UGameplayTagsManager_GameplayTagsMatch);
-		bool bResult;
-		if (MatchTypeOne == EGameplayTagMatchType::Explicit && MatchTypeTwo == EGameplayTagMatchType::Explicit)
-		{
-			bResult = GameplayTagOne == GameplayTagTwo;
-		}
-		else
-		{
-			bResult = ComplexGameplayTagsMatch(GameplayTagOne, MatchTypeOne, GameplayTagTwo, MatchTypeTwo);
-		}
-#if CHECK_TAG_OPTIMIZATIONS
-		check(bResult == GameplayTagsMatchOriginal(GameplayTagOne, MatchTypeOne, GameplayTagTwo, MatchTypeTwo));
-#endif
-		return bResult;
-	}
-
-	/**
-	 * Same as above except MatchTypeOne == EGameplayTagMatchType::Explicit && MatchTypeTwo == EGameplayTagMatchType::Explicit is not tested or handled
+	 *
+	 * @return the length of the longest matching pair
 	 */
-	bool ComplexGameplayTagsMatch(const FGameplayTag& GameplayTagOne, TEnumAsByte<EGameplayTagMatchType::Type> MatchTypeOne, const FGameplayTag& GameplayTagTwo, TEnumAsByte<EGameplayTagMatchType::Type> MatchTypeTwo) const;
-#if CHECK_TAG_OPTIMIZATIONS
-	bool GameplayTagsMatchOriginal(const FGameplayTag& GameplayTagOne, TEnumAsByte<EGameplayTagMatchType::Type> MatchTypeOne, const FGameplayTag& GameplayTagTwo, TEnumAsByte<EGameplayTagMatchType::Type> MatchTypeTwo) const;
-#endif
-
-	/**
-	* Check to see how closely two FGameplayTags match. Higher values indicate more matching terms in the tags.
-	*
-	* @param GameplayTagOne	The first tag to compare
-	* @param GameplayTagTwo	The second tag to compare
-	*
-	* @return the length of the longest matching pair
-	*/
 	int32 GameplayTagsMatchDepth(const FGameplayTag& GameplayTagOne, const FGameplayTag& GameplayTagTwo) const;
 
-	/**
-	 * Helper function for GameplayTagsMatch to get a container with all parents.
-	 * @param GameplayTag		tag to get parents of
-	 * @return					pointer to container of parents, if any
-	 */
-	FORCEINLINE const FGameplayTagContainer* GetAllParentsContainer(const FGameplayTag& GameplayTag)
-	{
-		const TSharedPtr<FGameplayTagNode>* TagNode = GameplayTagNodeMap.Find(GameplayTag);
-		if (TagNode)
-		{
-			return &((*TagNode)->Parents);
-		}
-#if !(UE_BUILD_SHIPPING || UE_BUILD_TEST)
-		UE_LOG(LogGameplayTags, Warning, TEXT("::GetAllParentsContainer Failed to find tag %s in GameplayTagNodeMap. Is this tag not in the dictionary?"), *GameplayTag.ToString());
-#endif
-		return nullptr;
-	}
-
-
-	/** Event for when assets are added to the registry */
-	DECLARE_EVENT(UGameplayTagsManager, FGameplayTagTreeChanged);
-	virtual FGameplayTagTreeChanged& OnGameplayTagTreeChanged() { return GameplayTagTreeChangedEvent; }
-
 	/** Returns true if we should import tags from UGameplayTagsSettings objects (configured by INI files) */
-	static bool ShouldImportTagsFromINI();
+	bool ShouldImportTagsFromINI();
 
-	/** TEMP - Returns true if we should warn on invalid (missing) tags */
-	static bool ShouldWarnOnInvalidTags();
+	/** Should we print loading errors when trying to load invalid tags */
+	bool ShouldWarnOnInvalidTags()
+	{
+		return bShouldWarnOnInvalidTags;
+	}
 
 	/** Should use fast replication */
 	bool ShouldUseFastReplication()
@@ -347,7 +297,7 @@ class GAMEPLAYTAGS_API UGameplayTagsManager : public UObject
 	}
 
 	/** Handles redirectors for an entire container, will also error on invalid tags */
-	void RedirectTagsForContainer(FGameplayTagContainer& Container, TSet<FName>& DeprecatedTagNamesNotFoundInTagMap, UProperty* SerializingProperty);
+	void RedirectTagsForContainer(FGameplayTagContainer& Container, UProperty* SerializingProperty);
 
 	/** Handles redirectors for a single tag, will also error on invalid tag. This is only called for when individual tags are serialized on their own */
 	void RedirectSingleGameplayTag(FGameplayTag& Tag, UProperty* SerializingProperty);
@@ -357,52 +307,105 @@ class GAMEPLAYTAGS_API UGameplayTagsManager : public UObject
 	FGameplayTagNetIndex GetNetIndexFromTag(const FGameplayTag &InTag);
 
 	/** Cached number of bits we need to replicate tags. That is, Log2(Number of Tags). Will always be <= 16. */
-	static int32 NetIndexTrueBitNum;
+	int32 NetIndexTrueBitNum;
 	
 	/** The length in bits of the first segment when net serializing tags. We will serialize NetIndexFirstBitSegment + 1 bit to indicatore "more" (more = second segment that is NetIndexTrueBitNum - NetIndexFirstBitSegment) */
-	static int32 NetIndexFirstBitSegment;
+	int32 NetIndexFirstBitSegment;
 
 	/** Numbers of bits to use for replicating container size. This can be set via config. */
-	static int32 NumBitsForContainerSize;
+	int32 NumBitsForContainerSize;
 
 	/** This is the actual value for an invalid tag "None". This is computed at runtime as (Total number of tags) + 1 */
 	FGameplayTagNetIndex InvalidTagNetIndex;
 
-	/** Cached runtime value for whether we are using fast replication or not. Initialized from config setting. */
-	bool bUseFastReplication;
-	
 	const TArray<TSharedPtr<FGameplayTagNode>>& GetNetworkGameplayTagNodeIndex() const { return NetworkGameplayTagNodeIndex; }
 
+#if WITH_EDITOR
+	/** Gets a Filtered copy of the GameplayRootTags Array based on the comma delimited filter string passed in */
+	void GetFilteredGameplayRootTags( const FString& InFilterString, TArray< TSharedPtr<FGameplayTagNode> >& OutTagArray );
+
+	/** Gets a list of all gameplay tag nodes added by the specific source */
+	void GetAllTagsFromSource(FName TagSource, TArray< TSharedPtr<FGameplayTagNode> >& OutTagArray);
+
+	/** Returns true if this tag is directly in the dictionary already */
+	bool IsDictionaryTag(FName TagName) const;
+
+	/** Returns comment and source for tag. If not found return false */
+	bool GetTagEditorData(FName TagName, FString& OutComment, FName &OutTagSource) const;
+
+	/** Refresh the gameplaytag tree due to an editor change */
+	void EditorRefreshGameplayTagTree();
+
+	/** Gets a Tag Container containing the all tags in the hierarchy that are children of this tag, and were explicitly added to the dictionary */
+	FGameplayTagContainer RequestGameplayTagChildrenInDictionary(const FGameplayTag& GameplayTag) const;
+
+#endif //WITH_EDITOR
+
+	DEPRECATED(4.15, "Call MatchesTag on FGameplayTag instead")
+	FORCEINLINE_DEBUGGABLE bool GameplayTagsMatch(const FGameplayTag& GameplayTagOne, TEnumAsByte<EGameplayTagMatchType::Type> MatchTypeOne, const FGameplayTag& GameplayTagTwo, TEnumAsByte<EGameplayTagMatchType::Type> MatchTypeTwo) const
+	{
+		SCOPE_CYCLE_COUNTER(STAT_UGameplayTagsManager_GameplayTagsMatch);
+		bool bResult = false;
+		if (MatchTypeOne == EGameplayTagMatchType::Explicit && MatchTypeTwo == EGameplayTagMatchType::Explicit)
+		{
+			bResult = GameplayTagOne == GameplayTagTwo;
+		}
+		else
+		{
+			// Convert both to their containers and do that match
+			const FGameplayTagContainer* ContainerOne = GetSingleTagContainer(GameplayTagOne);
+			const FGameplayTagContainer* ContainerTwo = GetSingleTagContainer(GameplayTagTwo);
+			if (ContainerOne && ContainerTwo)
+			{
+				bResult = ContainerOne->DoesTagContainerMatch(*ContainerTwo, MatchTypeOne, MatchTypeTwo, EGameplayContainerMatchType::Any);
+			}
+		}
+		return bResult;
+	}
+
 #if !(UE_BUILD_SHIPPING || UE_BUILD_TEST)
 	/** Mechanism for tracking what tags are frequently replicated */
 
-	static void PrintReplicationFrequencyReport();
-	static void NotifyTagReplicated(FGameplayTag Tag, bool WasInContainer);
+	void PrintReplicationFrequencyReport();
+	void NotifyTagReplicated(FGameplayTag Tag, bool WasInContainer);
 
-	static TMap<FGameplayTag, int32>	ReplicationCountMap;
-	static TMap<FGameplayTag, int32>	ReplicationCountMap_SingleTags;
-	static TMap<FGameplayTag, int32>	ReplicationCountMap_Containers;
+	TMap<FGameplayTag, int32>	ReplicationCountMap;
+	TMap<FGameplayTag, int32>	ReplicationCountMap_SingleTags;
+	TMap<FGameplayTag, int32>	ReplicationCountMap_Containers;
 #endif
 
 private:
 
-	friend class FGameplayTagTest;
-	friend class FGameplayEffectsTest;
+	/** Initializes the manager */
+	static void InitializeManager();
 
-	/** Helper function to populate the tag tree from each table */
-	void PopulateTreeFromDataTable(class UDataTable* Table);
+	/** The Tag Manager singleton */
+	static UGameplayTagsManager* SingletonManager;
 
-	void AddTagTableRow(const FGameplayTagTableRow& TagRow);
+	friend class FGameplayTagTest;
+	friend class FGameplayEffectsTest;
+	friend class FGameplayTagsModule;
+	friend class FGameplayTagsEditorModule;
 
 	/**
-	 * Helper function for RequestGameplayTagParents to add all parents to the container
+	 * Helper function to insert a tag into a tag node array
 	 *
-	 * @param TagContainer The container we need to add the parents too
-	 * @param GameplayTag The parent we need to check and add to the container
+	 * @param Tag					Tag to insert
+	 * @param ParentNode			Parent node, if any, for the tag
+	 * @param NodeArray				Node array to insert the new node into, if necessary (if the tag already exists, no insertion will occur)
+	 * @param SourceName			File tag was added from
+	 * @param DevComment			Comment from developer about this tag
+	 *
+	 * @return Index of the node of the tag
 	 */
-	void AddParentTags(FGameplayTagContainer& TagContainer, const FGameplayTag& GameplayTag) const;
+	int32 InsertTagIntoNodeArray(FName Tag, TSharedPtr<FGameplayTagNode> ParentNode, TArray< TSharedPtr<FGameplayTagNode> >& NodeArray, FName SourceName, const FString& DevComment);
+
+	/** Helper function to populate the tag tree from each table */
+	void PopulateTreeFromDataTable(class UDataTable* Table);
+
+	void AddTagTableRow(const FGameplayTagTableRow& TagRow, FName SourceName);
 
-	void AddChildrenTags(FGameplayTagContainer& TagContainer, const FGameplayTag& GameplayTag, bool RecurseAll=true, bool OnlyIncludeDictionaryTags=false) const;
+	void AddChildrenTags(FGameplayTagContainer& TagContainer, const TSharedPtr<FGameplayTagNode> GameplayTagNode, bool RecurseAll=true, bool OnlyIncludeDictionaryTags=false) const;
 
 	/**
 	 * Helper function for GameplayTagsMatch to get all parents when doing a parent match,
@@ -413,49 +416,30 @@ private:
 	 */
 	void GetAllParentNodeNames(TSet<FName>& NamesList, const TSharedPtr<FGameplayTagNode> GameplayTag) const;
 
+	/** Returns the tag source for a given tag source name, or null if not found */
+	FGameplayTagSource* FindOrAddTagSource(FName TagSourceName, EGameplayTagSourceType SourceType);
+
 	/** Constructs the net indices for each tag */
 	void ConstructNetIndex();
 
 	/** Roots of gameplay tag nodes */
 	TSharedPtr<FGameplayTagNode> GameplayRootTag;
 
-	class GamePlayTagMapWithInverse
-	{
-	public:
-		void Add(FGameplayTag Tag, TSharedPtr<FGameplayTagNode> GameplayTagNode)
-		{
-			ForwardMap.Add(Tag, GameplayTagNode);
-			InverseMap.Add(GameplayTagNode, Tag);
-		}
-
-		const TSharedPtr<FGameplayTagNode>* Find(const FGameplayTag& Tag) const
-		{
-			return ForwardMap.Find(Tag);
-		}
+	/** Map of Tags to Nodes - Internal use only. FGameplayTag is inside node structure, do not use FindKey! */
+	TMap<FGameplayTag, TSharedPtr<FGameplayTagNode>> GameplayTagNodeMap;
 
-		const FGameplayTag* FindKey(const TSharedPtr<FGameplayTagNode>& TagNode) const
-		{
-			return InverseMap.Find(TagNode);
-		}
-
-		void GenerateValueArray(TArray<TSharedPtr<FGameplayTagNode>>& ValueArray) const
-		{
-			ForwardMap.GenerateValueArray(ValueArray);
-		}
+	/** Our aggregated, sorted list of commonly replicated tags. These tags are given lower indices to ensure they replicate in the first bit segment. */
+	TArray<FGameplayTag> CommonlyReplicatedTags;
 
-	private:
-		TMap<FGameplayTag, TSharedPtr<FGameplayTagNode>> ForwardMap;
-		TMap<TSharedPtr<FGameplayTagNode>, FGameplayTag> InverseMap;
-	};
+	/** List of gameplay tag sources */
+	UPROPERTY()
+	TArray<FGameplayTagSource> TagSources;
 
-	/** Map of Tags to Nodes - Internal use only */
-	GamePlayTagMapWithInverse GameplayTagNodeMap;
+	/** Cached runtime value for whether we are using fast replication or not. Initialized from config setting. */
+	bool bUseFastReplication;
 
-	/** Map of Names to tags - Internal use only */
-	TMap<FName, FGameplayTag> GameplayTagMap;
-
-	/** Our aggregated, sorted list of commonly replicated tags. These tags are given lower indices to ensure they replicate in the first bit segment. */
-	TArray<FGameplayTag> CommonlyReplicatedTags;
+	/** Cached runtime value for whether we should warn when loading invalid tags */
+	bool bShouldWarnOnInvalidTags;
 
 #if WITH_EDITOR
 	// This critical section is to handle and editor-only issue where tag requests come from another thread when async loading from a background thread in FGameplayTagContainer::Serialize.
@@ -470,14 +454,6 @@ private:
 	UPROPERTY()
 	TArray<UDataTable*> GameplayTagTables;
 
-	/** The delegate to execute when the tag tree changes */
-	FGameplayTagTreeChanged GameplayTagTreeChangedEvent;
-
 	/** The map of ini-configured tag redirectors */
 	TMap<FName, FGameplayTag> TagRedirects;
-
-#if WITH_EDITOR
-	/** Flag to say if we have registered the ObjectReimport, this is needed as use of Tags manager can happen before Editor is ready*/
-	bool RegisteredObjectReimport;
-#endif
-};
+};
\ No newline at end of file
diff --git a/Engine/Source/Runtime/GameplayTags/Classes/GameplayTagsSettings.h b/Engine/Source/Runtime/GameplayTags/Classes/GameplayTagsSettings.h
--- a/Engine/Source/Runtime/GameplayTags/Classes/GameplayTagsSettings.h
+++ b/Engine/Source/Runtime/GameplayTags/Classes/GameplayTagsSettings.h
@@ -5,10 +5,45 @@
 #include "GameplayTagContainer.h"
 #include "GameplayTagsSettings.generated.h"
 
+USTRUCT()
+struct GAMEPLAYTAGS_API FGameplayTagRedirect
+{
+	GENERATED_USTRUCT_BODY()
+
+	UPROPERTY(EditAnywhere, Category = GameplayTags)
+	FName OldTagName;
+
+	UPROPERTY(EditAnywhere, Category = GameplayTags)
+	FName NewTagName;
+
+	friend inline bool operator==(const FGameplayTagRedirect& A, const FGameplayTagRedirect& B)
+	{
+		return A.OldTagName == B.OldTagName && A.NewTagName == B.NewTagName;
+	}
+};
+
+/** Base class for storing a list of gameplay tags as an ini list */
+UCLASS(config = GameplayTagsList, notplaceable)
+class GAMEPLAYTAGS_API UGameplayTagsList : public UObject
+{
+	GENERATED_UCLASS_BODY()
+
+	/** Relative path to the ini file that is backing this list */
+	UPROPERTY()
+	FString ConfigFileName;
+
+	/** List of tags saved to this file */
+	UPROPERTY(config, EditAnywhere, Category = GameplayTags)
+	TArray<FGameplayTagTableRow> GameplayTagList;
+
+	/** Sorts tags alphabetically */
+	void SortTags();
+};
+
 /**
  *	Class for importing GameplayTags directly from a config file.
  *	FGameplayTagsEditorModule::StartupModule adds this class to the Project Settings menu to be edited.
- *	Editing this in Project Settings will output changes to Config/DEfaultEngine.ini.
+ *	Editing this in Project Settings will output changes to Config/DefaultGameplayTags.ini.
  *	
  *	Primary advantages of this approach are:
  *	-Adding new tags doesn't require checking out external and editing file (CSV or xls) then reimporting.
@@ -20,22 +55,41 @@
  *	
  */
 UCLASS(config=GameplayTags, defaultconfig, notplaceable)
-class GAMEPLAYTAGS_API UGameplayTagsSettings : public UObject
+class GAMEPLAYTAGS_API UGameplayTagsSettings : public UGameplayTagsList
 {
 	GENERATED_UCLASS_BODY()
 
-	UPROPERTY(config, EditAnywhere, Category=GameplayTags)
-	TArray<FString>		GameplayTags;
+	/** List of active tag redirects */
+	UPROPERTY(config, EditAnywhere, Category = GameplayTags)
+	TArray<FGameplayTagRedirect> GameplayTagRedirects;
+
+	/** List of data tables to load tags from */
+	UPROPERTY(config, EditAnywhere, Category = GameplayTags, meta = (AllowedClasses = "DataTable"))
+	TArray<FStringAssetReference> GameplayTagTableList;
+
+	/** If true, will import tags from ini files in the config/tags folder */
+	UPROPERTY(config, EditAnywhere, Category = GameplayTags)
+	bool ImportTagsFromConfig;
+
+	/** If true, will give load warnings when reading invalid tags off disk */
+	UPROPERTY(config, EditAnywhere, Category = GameplayTags)
+	bool WarnOnInvalidTags;
+
+	/** If true, will replicate gameplay tags by index instead of name. For this to work, tags must be identical on client and server */
+	UPROPERTY(config, EditAnywhere, Category = "Advanced Replication")
+	bool FastReplication;
 
 	/** List of tags most frequently replicated */
-	UPROPERTY(config, EditAnywhere, Category="Advanced Replication")
-	TArray<FString>		CommonlyReplicatedTags;
+	UPROPERTY(config, EditAnywhere, Category = "Advanced Replication")
+	TArray<FName> CommonlyReplicatedTags;
 
-	UPROPERTY(config, EditAnywhere, Category="Advanced Replication")
-	int32 NetIndexFirstBitSegment;
+	/** Numbers of bits to use for replicating container size, set this based on how large your containers tend to be */
+	UPROPERTY(config, EditAnywhere, Category = "Advanced Replication")
+	int32 NumBitsForContainerSize;
 
-	/** Sorts tags alphabetically */
-	void SortTags();
+	/** The length in bits of the first segment when net serializing tags. We will serialize NetIndexFirstBitSegment + 1 bit to indicate "more", which is slower to replicate */
+	UPROPERTY(config, EditAnywhere, Category= "Advanced Replication")
+	int32 NetIndexFirstBitSegment;
 
 #if WITH_EDITOR
 	virtual void PostEditChangeProperty(FPropertyChangedEvent& PropertyChangedEvent) override;
diff --git a/Engine/Source/Runtime/GameplayTags/GameplayTags.Build.cs b/Engine/Source/Runtime/GameplayTags/GameplayTags.Build.cs
--- a/Engine/Source/Runtime/GameplayTags/GameplayTags.Build.cs
+++ b/Engine/Source/Runtime/GameplayTags/GameplayTags.Build.cs
@@ -6,16 +6,9 @@ namespace UnrealBuildTool.Rules
 	{
 		public GameplayTags(TargetInfo Target)
 		{
-			PublicIncludePaths.AddRange(
-				new string[] {
-					// ... add public include paths required here ...
-				}
-				);
-
 			PrivateIncludePaths.AddRange(
 				new string[] {
 					"Runtime/GameplayTags/Private",
-					// ... add other private include paths required here ...
 				}
 				);
 
@@ -25,31 +18,8 @@ namespace UnrealBuildTool.Rules
 					"Core",
 					"CoreUObject",
 					"Engine"
-					// ... add other public dependencies that you statically link with here ...
 				}
 				);
-
-			PrivateDependencyModuleNames.AddRange(
-				new string[]
-				{
-					// ... add private dependencies that you statically link with here ...
-				}
-				);
-
-			DynamicallyLoadedModuleNames.AddRange(
-				new string[]
-				{
-					// ... add any modules that your module loads dynamically here ...
-				}
-				);
-
-			if (UEBuildConfiguration.bBuildEditor == true)
-			{
-				PrivateDependencyModuleNames.Add("UnrealEd");
-                PrivateDependencyModuleNames.Add("SourceControl");
-                PrivateDependencyModuleNames.Add("Slate");
-                PrivateDependencyModuleNames.Add("SlateCore");
-			}
 		}
 	}
 }
\ No newline at end of file
diff --git a/Engine/Source/Runtime/GameplayTags/Private/BlueprintGameplayTagLibrary.cpp b/Engine/Source/Runtime/GameplayTags/Private/BlueprintGameplayTagLibrary.cpp
--- a/Engine/Source/Runtime/GameplayTags/Private/BlueprintGameplayTagLibrary.cpp
+++ b/Engine/Source/Runtime/GameplayTags/Private/BlueprintGameplayTagLibrary.cpp
@@ -7,9 +7,48 @@ UBlueprintGameplayTagLibrary::UBlueprintGameplayTagLibrary(const FObjectInitiali
 {
 }
 
-bool UBlueprintGameplayTagLibrary::DoGameplayTagsMatch(const FGameplayTag& TagOne, const FGameplayTag& TagTwo, TEnumAsByte<EGameplayTagMatchType::Type> TagOneMatchType, TEnumAsByte<EGameplayTagMatchType::Type> TagTwoMatchType)
+bool UBlueprintGameplayTagLibrary::MatchesTag(FGameplayTag TagOne, FGameplayTag TagTwo, bool bExactMatch)
 {
-	return TagOne.Matches(TagOneMatchType, TagTwo, TagTwoMatchType);
+	if (bExactMatch)
+	{
+		return TagOne.MatchesTagExact(TagTwo);
+	}
+
+	return TagOne.MatchesTag(TagTwo);
+}
+
+bool UBlueprintGameplayTagLibrary::MatchesAnyTags(FGameplayTag TagOne, const FGameplayTagContainer& OtherContainer, bool bExactMatch)
+{
+	if (bExactMatch)
+	{
+		return TagOne.MatchesAnyExact(OtherContainer);
+	}
+	return TagOne.MatchesAny(OtherContainer);
+}
+
+bool UBlueprintGameplayTagLibrary::EqualEqual_GameplayTag(FGameplayTag A, FGameplayTag B)
+{
+	return A == B;
+}
+
+bool UBlueprintGameplayTagLibrary::NotEqual_GameplayTag(FGameplayTag A, FGameplayTag B)
+{
+	return A != B;
+}
+
+bool UBlueprintGameplayTagLibrary::IsGameplayTagValid(FGameplayTag GameplayTag)
+{
+	return GameplayTag.IsValid();
+}
+
+FName UBlueprintGameplayTagLibrary::GetTagName(const FGameplayTag& GameplayTag)
+{
+	return GameplayTag.GetTagName();
+}
+
+FGameplayTag UBlueprintGameplayTagLibrary::MakeLiteralGameplayTag(FGameplayTag Value)
+{
+	return Value;
 }
 
 int32 UBlueprintGameplayTagLibrary::GetNumGameplayTagsInContainer(const FGameplayTagContainer& TagContainer)
@@ -17,19 +56,31 @@ int32 UBlueprintGameplayTagLibrary::GetNumGameplayTagsInContainer(const FGamepla
 	return TagContainer.Num();
 }
 
-bool UBlueprintGameplayTagLibrary::DoesContainerHaveTag(const FGameplayTagContainer& TagContainer, TEnumAsByte<EGameplayTagMatchType::Type> ContainerTagsMatchType, const FGameplayTag& Tag, TEnumAsByte<EGameplayTagMatchType::Type> TagMatchType)
+bool UBlueprintGameplayTagLibrary::HasTag(const FGameplayTagContainer& TagContainer, FGameplayTag Tag, bool bExactMatch)
 {
-	return (TagContainer.HasTag(Tag, ContainerTagsMatchType, TagMatchType));
+	if (bExactMatch)
+	{
+		return TagContainer.HasTagExact(Tag);
+	}
+	return TagContainer.HasTag(Tag);
 }
 
-bool UBlueprintGameplayTagLibrary::DoesContainerMatchAnyTagsInContainer(const FGameplayTagContainer& TagContainer, const FGameplayTagContainer& OtherContainer, bool bCountEmptyAsMatch)
+bool UBlueprintGameplayTagLibrary::HasAnyTags(const FGameplayTagContainer& TagContainer, const FGameplayTagContainer& OtherContainer, bool bExactMatch)
 {
-	return TagContainer.MatchesAny(OtherContainer, bCountEmptyAsMatch);
+	if (bExactMatch)
+	{
+		return TagContainer.HasAnyExact(OtherContainer);
+	}
+	return TagContainer.HasAny(OtherContainer);
 }
 
-bool UBlueprintGameplayTagLibrary::DoesContainerMatchAllTagsInContainer(const FGameplayTagContainer& TagContainer, const FGameplayTagContainer& OtherContainer, bool bCountEmptyAsMatch)
+bool UBlueprintGameplayTagLibrary::HasAllTags(const FGameplayTagContainer& TagContainer, const FGameplayTagContainer& OtherContainer, bool bExactMatch)
 {
-	return TagContainer.MatchesAll(OtherContainer, bCountEmptyAsMatch);
+	if (bExactMatch)
+	{
+		return TagContainer.HasAllExact(OtherContainer);
+	}
+	return TagContainer.HasAll(OtherContainer);
 }
 
 bool UBlueprintGameplayTagLibrary::DoesContainerMatchTagQuery(const FGameplayTagContainer& TagContainer, const FGameplayTagQuery& TagQuery)
@@ -37,31 +88,54 @@ bool UBlueprintGameplayTagLibrary::DoesContainerMatchTagQuery(const FGameplayTag
 	return TagQuery.Matches(TagContainer);
 }
 
-FGameplayTag UBlueprintGameplayTagLibrary::MakeLiteralGameplayTag(FGameplayTag Value)
+bool UBlueprintGameplayTagLibrary::EqualEqual_GameplayTagContainer(const FGameplayTagContainer& A, const FGameplayTagContainer& B)
+{
+	return A == B;
+}
+
+bool UBlueprintGameplayTagLibrary::NotEqual_GameplayTagContainer(const FGameplayTagContainer& A, const FGameplayTagContainer& B)
+{
+	return A != B;
+}
+
+FGameplayTagContainer UBlueprintGameplayTagLibrary::MakeLiteralGameplayTagContainer(FGameplayTagContainer Value)
 {
 	return Value;
 }
 
+FGameplayTagContainer UBlueprintGameplayTagLibrary::MakeGameplayTagContainerFromArray(const TArray<FGameplayTag>& GameplayTags)
+{
+	return FGameplayTagContainer::CreateFromArray(GameplayTags);
+}
+
+FGameplayTagContainer UBlueprintGameplayTagLibrary::MakeGameplayTagContainerFromTag(FGameplayTag SingleTag)
+{
+	return FGameplayTagContainer(SingleTag);
+}
+
+void UBlueprintGameplayTagLibrary::BreakGameplayTagContainer(const FGameplayTagContainer& GameplayTagContainer, TArray<FGameplayTag>& GameplayTags)
+{
+	GameplayTagContainer.GetGameplayTagArray(GameplayTags);
+}
+
 FGameplayTagQuery UBlueprintGameplayTagLibrary::MakeGameplayTagQuery(FGameplayTagQuery TagQuery)
 {
 	return TagQuery;
 }
 
-bool UBlueprintGameplayTagLibrary::HasAllMatchingGameplayTags(TScriptInterface<IGameplayTagAssetInterface> TagContainerInterface, const FGameplayTagContainer& OtherContainer, bool bCountEmptyAsMatch)
+bool UBlueprintGameplayTagLibrary::HasAllMatchingGameplayTags(TScriptInterface<IGameplayTagAssetInterface> TagContainerInterface, const FGameplayTagContainer& OtherContainer)
 {
 	if (TagContainerInterface.GetInterface() == NULL)
 	{
-		if (bCountEmptyAsMatch)
-		{
-			return (OtherContainer.Num() == 0);
-		}
-		return false;
+		return (OtherContainer.Num() == 0);
 	}
 
-	return TagContainerInterface->HasAllMatchingGameplayTags(OtherContainer, bCountEmptyAsMatch);
+	FGameplayTagContainer OwnedTags;
+	TagContainerInterface->GetOwnedGameplayTags(OwnedTags);
+	return (OwnedTags.HasAll(OtherContainer));
 }
 
-bool UBlueprintGameplayTagLibrary::DoesTagAssetInterfaceHaveTag(TScriptInterface<IGameplayTagAssetInterface> TagContainerInterface, TEnumAsByte<EGameplayTagMatchType::Type> ContainerTagsMatchType, const FGameplayTag& Tag, TEnumAsByte<EGameplayTagMatchType::Type> TagMatchType)
+bool UBlueprintGameplayTagLibrary::DoesTagAssetInterfaceHaveTag(TScriptInterface<IGameplayTagAssetInterface> TagContainerInterface, FGameplayTag Tag)
 {
 	if (TagContainerInterface.GetInterface() == NULL)
 	{
@@ -70,14 +144,17 @@ bool UBlueprintGameplayTagLibrary::DoesTagAssetInterfaceHaveTag(TScriptInterface
 
 	FGameplayTagContainer OwnedTags;
 	TagContainerInterface->GetOwnedGameplayTags(OwnedTags);
-	return (OwnedTags.HasTag(Tag, ContainerTagsMatchType, TagMatchType));
+	return (OwnedTags.HasTag(Tag));
 }
 
-bool UBlueprintGameplayTagLibrary::AppendGameplayTagContainers(const FGameplayTagContainer& InTagContainer, UPARAM(ref) FGameplayTagContainer& InOutTagContainer)
+void UBlueprintGameplayTagLibrary::AppendGameplayTagContainers(FGameplayTagContainer& InOutTagContainer, const FGameplayTagContainer& InTagContainer)
 {
 	InOutTagContainer.AppendTags(InTagContainer);
+}
 
-	return true;
+void UBlueprintGameplayTagLibrary::AddGameplayTag(FGameplayTagContainer& InOutTagContainer, FGameplayTag Tag)
+{
+	InOutTagContainer.AddTag(Tag);
 }
 
 bool UBlueprintGameplayTagLibrary::NotEqual_TagTag(FGameplayTag A, FString B)
@@ -126,7 +203,7 @@ bool UBlueprintGameplayTagLibrary::NotEqual_TagContainerTagContainer(FGameplayTa
 			}
 			TagString = Remainder;
 
-			const FGameplayTag Tag = IGameplayTagsModule::Get().GetGameplayTagsManager().RequestGameplayTag(FName(*ReadTag));
+			const FGameplayTag Tag = FGameplayTag::RequestGameplayTag(FName(*ReadTag));
 			TagContainer.AddTag(Tag);
 		}
 		if (Remainder.IsEmpty())
@@ -145,7 +222,7 @@ bool UBlueprintGameplayTagLibrary::NotEqual_TagContainerTagContainer(FGameplayTa
 					Remainder = Remainder.RightChop(1);
 				}
 			}
-			const FGameplayTag Tag = IGameplayTagsModule::Get().GetGameplayTagsManager().RequestGameplayTag(FName(*Remainder));
+			const FGameplayTag Tag = FGameplayTag::RequestGameplayTag(FName(*Remainder));
 			TagContainer.AddTag(Tag);
 		}
 	}
@@ -157,12 +234,7 @@ FString UBlueprintGameplayTagLibrary::GetDebugStringFromGameplayTagContainer(con
 	return TagContainer.ToStringSimple();
 }
 
-FString UBlueprintGameplayTagLibrary::GetDebugStringFromGameplayTag(const FGameplayTag& GameplayTag)
+FString UBlueprintGameplayTagLibrary::GetDebugStringFromGameplayTag(FGameplayTag GameplayTag)
 {
 	return GameplayTag.ToString();
 }
-
-bool UBlueprintGameplayTagLibrary::IsGameplayTagValid(const FGameplayTag& TagContainer)
-{
-	return TagContainer.IsValid();
-}
\ No newline at end of file
diff --git a/Engine/Source/Runtime/GameplayTags/Private/GameplayTagAssetInterface.cpp b/Engine/Source/Runtime/GameplayTags/Private/GameplayTagAssetInterface.cpp
--- a/Engine/Source/Runtime/GameplayTags/Private/GameplayTagAssetInterface.cpp
+++ b/Engine/Source/Runtime/GameplayTags/Private/GameplayTagAssetInterface.cpp
@@ -12,21 +12,21 @@ bool IGameplayTagAssetInterface::HasMatchingGameplayTag(FGameplayTag TagToCheck)
 	FGameplayTagContainer OwnedTags;
 	GetOwnedGameplayTags(OwnedTags);
 
-	return OwnedTags.HasTag(TagToCheck, EGameplayTagMatchType::IncludeParentTags, EGameplayTagMatchType::Explicit);
+	return OwnedTags.HasTag(TagToCheck);
 }
 
-bool IGameplayTagAssetInterface::HasAllMatchingGameplayTags(const FGameplayTagContainer& TagContainer, bool bCountEmptyAsMatch) const
+bool IGameplayTagAssetInterface::HasAllMatchingGameplayTags(const FGameplayTagContainer& TagContainer) const
 {
 	FGameplayTagContainer OwnedTags;
 	GetOwnedGameplayTags(OwnedTags);
 
-	return OwnedTags.MatchesAll(TagContainer, bCountEmptyAsMatch);
+	return OwnedTags.HasAll(TagContainer);
 }
 
-bool IGameplayTagAssetInterface::HasAnyMatchingGameplayTags(const FGameplayTagContainer& TagContainer, bool bCountEmptyAsMatch) const
+bool IGameplayTagAssetInterface::HasAnyMatchingGameplayTags(const FGameplayTagContainer& TagContainer) const
 {
 	FGameplayTagContainer OwnedTags;
 	GetOwnedGameplayTags(OwnedTags);
 
-	return OwnedTags.MatchesAny(TagContainer, bCountEmptyAsMatch);
+	return OwnedTags.HasAny(TagContainer);
 }
