Copying //UE4/Dev-Framework to //UE4/Dev-Main (Source: //UE4/Dev-Framework @ 3182037)

diff --git a/Engine/Source/Runtime/GameplayTags/Private/GameplayTagContainer.cpp b/Engine/Source/Runtime/GameplayTags/Private/GameplayTagContainer.cpp
--- a/Engine/Source/Runtime/GameplayTags/Private/GameplayTagContainer.cpp
+++ b/Engine/Source/Runtime/GameplayTags/Private/GameplayTagContainer.cpp
@@ -4,6 +4,7 @@
 #include "PropertyTag.h"
 #include "Engine/PackageMapClient.h"
 
+const FGameplayTag FGameplayTag::EmptyTag;
 const FGameplayTagContainer FGameplayTagContainer::EmptyContainer;
 const FGameplayTagQuery FGameplayTagQuery::EmptyQuery;
 
@@ -267,7 +268,7 @@ bool FQueryEvaluator::EvalAnyTagsMatch(FGameplayTagContainer const& Tags, bool b
 		{
 			FGameplayTag Tag = Query.GetTagFromIndex(TagIdx);
 
-			bool bHasTag = Tags.HasTag(Tag, EGameplayTagMatchType::IncludeParentTags, EGameplayTagMatchType::Explicit);
+			bool bHasTag = Tags.HasTag(Tag);
 
 			if (bHasTag)
 			{
@@ -306,7 +307,7 @@ bool FQueryEvaluator::EvalAllTagsMatch(FGameplayTagContainer const& Tags, bool b
 		if (bShortCircuit == false)
 		{
 			FGameplayTag const Tag = Query.GetTagFromIndex(TagIdx);
-			bool const bHasTag = Tags.HasTag(Tag, EGameplayTagMatchType::IncludeParentTags, EGameplayTagMatchType::Explicit);
+			bool const bHasTag = Tags.HasTag(Tag);
 
 			if (bHasTag == false)
 			{
@@ -345,7 +346,7 @@ bool FQueryEvaluator::EvalNoTagsMatch(FGameplayTagContainer const& Tags, bool bS
 		if (bShortCircuit == false)
 		{
 			FGameplayTag const Tag = Query.GetTagFromIndex(TagIdx);
-			bool const bHasTag = Tags.HasTag(Tag, EGameplayTagMatchType::IncludeParentTags, EGameplayTagMatchType::Explicit);
+			bool const bHasTag = Tags.HasTag(Tag);
 
 			if (bHasTag == true)
 			{
@@ -492,53 +493,47 @@ FGameplayTagContainer& FGameplayTagContainer::operator=(FGameplayTagContainer co
 	GameplayTags.Empty(Other.GameplayTags.Num());
 	GameplayTags.Append(Other.GameplayTags);
 
+	ParentTags.Empty(Other.ParentTags.Num());
+	ParentTags.Append(Other.ParentTags);
+
 	return *this;
 }
 
 FGameplayTagContainer& FGameplayTagContainer::operator=(FGameplayTagContainer&& Other)
 {
 	GameplayTags = MoveTemp(Other.GameplayTags);
+	ParentTags = MoveTemp(Other.ParentTags);
 	return *this;
 }
 
 bool FGameplayTagContainer::operator==(FGameplayTagContainer const& Other) const
 {
+	// This is to handle the case where the two containers are in different orders
 	if (GameplayTags.Num() != Other.GameplayTags.Num())
 	{
 		return false;
 	}
-	return Filter(Other, EGameplayTagMatchType::Explicit, EGameplayTagMatchType::Explicit).Num() == this->Num();
+	return FilterExact(Other).Num() == this->Num();
 }
 
 bool FGameplayTagContainer::operator!=(FGameplayTagContainer const& Other) const
 {
-	if (GameplayTags.Num() != Other.GameplayTags.Num())
-	{
-		return true;
-	}
-	return Filter(Other, EGameplayTagMatchType::Explicit, EGameplayTagMatchType::Explicit).Num() != this->Num();
+	return !operator==(Other);
 }
 
 bool FGameplayTagContainer::ComplexHasTag(FGameplayTag const& TagToCheck, TEnumAsByte<EGameplayTagMatchType::Type> TagMatchType, TEnumAsByte<EGameplayTagMatchType::Type> TagToCheckMatchType) const
 {
 	check(TagMatchType != EGameplayTagMatchType::Explicit || TagToCheckMatchType != EGameplayTagMatchType::Explicit);
-	UGameplayTagsManager& TagManager = IGameplayTagsModule::GetGameplayTagsManager();
 
-	if (TagMatchType != EGameplayTagMatchType::Explicit)
+	if (TagMatchType == EGameplayTagMatchType::IncludeParentTags)
 	{
-		for (TArray<FGameplayTag>::TConstIterator It(this->GameplayTags); It; ++It)
-		{
-			const FGameplayTagContainer* Parents = TagManager.GetAllParentsContainer(*It);
-			if (Parents && Parents->HasTag(TagToCheck, EGameplayTagMatchType::Explicit, TagToCheckMatchType))
-			{
-				return true;
-			}
-		}
+		FGameplayTagContainer ExpandedConatiner = GetGameplayTagParents();
+		return ExpandedConatiner.HasTagFast(TagToCheck, EGameplayTagMatchType::Explicit, TagToCheckMatchType);
 	}
 	else
 	{
-		const FGameplayTagContainer* Parents = TagManager.GetAllParentsContainer(TagToCheck);
-		if (Parents && DoesTagContainerMatch(*Parents, EGameplayTagMatchType::Explicit, EGameplayTagMatchType::Explicit, EGameplayContainerMatchType::Any))
+		const FGameplayTagContainer* SingleContainer = UGameplayTagsManager::Get().GetSingleTagContainer(TagToCheck);
+		if (SingleContainer && SingleContainer->DoesTagContainerMatch(*this, EGameplayTagMatchType::IncludeParentTags, EGameplayTagMatchType::Explicit, EGameplayContainerMatchType::Any))
 		{
 			return true;
 		}
@@ -547,23 +542,6 @@ bool FGameplayTagContainer::ComplexHasTag(FGameplayTag const& TagToCheck, TEnumA
 	return false;
 }
 
-#if CHECK_TAG_OPTIMIZATIONS
-bool FGameplayTagContainer::HasTagOriginal(FGameplayTag const& TagToCheck, TEnumAsByte<EGameplayTagMatchType::Type> TagMatchType, TEnumAsByte<EGameplayTagMatchType::Type> TagToCheckMatchType) const
-{
-	SCOPE_CYCLE_COUNTER(STAT_FGameplayTagContainer_HasTag);
-
-	UGameplayTagsManager& TagManager = IGameplayTagsModule::GetGameplayTagsManager();
-	for (TArray<FGameplayTag>::TConstIterator It(this->GameplayTags); It; ++It)
-	{
-		if (TagManager.GameplayTagsMatch(*It, TagMatchType, TagToCheck, TagToCheckMatchType) == true)
-		{
-			return true;
-		}
-	}
-	return false;
-}
-#endif
-
 DECLARE_CYCLE_STAT(TEXT("FGameplayTagContainer::RemoveTagByExplicitName"), STAT_FGameplayTagContainer_RemoveTagByExplicitName, STATGROUP_GameplayTags);
 
 bool FGameplayTagContainer::RemoveTagByExplicitName(const FName& TagName)
@@ -582,17 +560,43 @@ bool FGameplayTagContainer::RemoveTagByExplicitName(const FName& TagName)
 	return false;
 }
 
-DECLARE_CYCLE_STAT(TEXT("FGameplayTagContainer::GetGameplayTagParents"), STAT_FGameplayTagContainer_GetGameplayTagParents, STATGROUP_GameplayTags);
+FORCEINLINE_DEBUGGABLE void FGameplayTagContainer::AddParentsForTag(const FGameplayTag& Tag)
+{
+	const FGameplayTagContainer* SingleContainer = UGameplayTagsManager::Get().GetSingleTagContainer(Tag);
+
+	if (SingleContainer)
+	{
+		// Add Parent tags from this tag to our own
+		for (const FGameplayTag& ParentTag : SingleContainer->ParentTags)
+		{
+			ParentTags.AddUnique(ParentTag);
+		}
+	}
+}
+
+DECLARE_CYCLE_STAT(TEXT("FGameplayTagContainer::FillParentTags"), STAT_FGameplayTagContainer_FillParentTags, STATGROUP_GameplayTags);
+
+void FGameplayTagContainer::FillParentTags()
+{
+	SCOPE_CYCLE_COUNTER(STAT_FGameplayTagContainer_FillParentTags);
+
+	ParentTags.Reset();
+
+	for (const FGameplayTag& Tag : GameplayTags)
+	{
+		AddParentsForTag(Tag);
+	}
+}
 
 FGameplayTagContainer FGameplayTagContainer::GetGameplayTagParents() const
 {
 	FGameplayTagContainer ResultContainer;
-	ResultContainer.AppendTags(*this);
+	ResultContainer.GameplayTags = GameplayTags;
 
-	for (TArray<FGameplayTag>::TConstIterator TagIterator(GameplayTags); TagIterator; ++TagIterator)
+	// Add parent tags to explicit tags, the rest got copied over already
+	for (const FGameplayTag& Tag : ParentTags)
 	{
-		FGameplayTagContainer ParentTags = IGameplayTagsModule::Get().GetGameplayTagsManager().RequestGameplayTagParents(*TagIterator);
-		ResultContainer.AppendTags(ParentTags);
+		ResultContainer.GameplayTags.AddUnique(Tag);
 	}
 
 	return ResultContainer;
@@ -605,17 +609,49 @@ FGameplayTagContainer FGameplayTagContainer::Filter(const FGameplayTagContainer&
 	SCOPE_CYCLE_COUNTER(STAT_FGameplayTagContainer_Filter);
 
 	FGameplayTagContainer ResultContainer;
-	UGameplayTagsManager& TagManager = IGameplayTagsModule::GetGameplayTagsManager();
 
-	
-	for (TArray<FGameplayTag>::TConstIterator OtherIt(OtherContainer.GameplayTags); OtherIt; ++OtherIt)
+	for (const FGameplayTag& Tag : GameplayTags)
 	{
-		for (TArray<FGameplayTag>::TConstIterator It(this->GameplayTags); It; ++It)
+		PRAGMA_DISABLE_DEPRECATION_WARNINGS
+		// Check to see if all of these tags match other container, with types swapped
+		if (OtherContainer.HasTag(Tag, OtherTagMatchType, TagMatchType))
 		{
-			if (TagManager.GameplayTagsMatch(*It, TagMatchType, *OtherIt, OtherTagMatchType) == true)
-			{
-				ResultContainer.AddTag(*It);
-			}
+			ResultContainer.AddTagFast(Tag);
+		}
+		PRAGMA_ENABLE_DEPRECATION_WARNINGS
+	}
+
+	return ResultContainer;
+}
+
+FGameplayTagContainer FGameplayTagContainer::Filter(const FGameplayTagContainer& OtherContainer) const
+{
+	SCOPE_CYCLE_COUNTER(STAT_FGameplayTagContainer_Filter);
+
+	FGameplayTagContainer ResultContainer;
+
+	for (const FGameplayTag& Tag : GameplayTags)
+	{
+		if (Tag.MatchesAny(OtherContainer))
+		{
+			ResultContainer.AddTagFast(Tag);
+		}
+	}
+
+	return ResultContainer;
+}
+
+FGameplayTagContainer FGameplayTagContainer::FilterExact(const FGameplayTagContainer& OtherContainer) const
+{
+	SCOPE_CYCLE_COUNTER(STAT_FGameplayTagContainer_Filter);
+
+	FGameplayTagContainer ResultContainer;
+
+	for (const FGameplayTag& Tag : GameplayTags)
+	{
+		if (Tag.MatchesAnyExact(OtherContainer))
+		{
+			ResultContainer.AddTagFast(Tag);
 		}
 	}
 
@@ -624,7 +660,7 @@ FGameplayTagContainer FGameplayTagContainer::Filter(const FGameplayTagContainer&
 
 bool FGameplayTagContainer::DoesTagContainerMatchComplex(const FGameplayTagContainer& OtherContainer, TEnumAsByte<EGameplayTagMatchType::Type> TagMatchType, TEnumAsByte<EGameplayTagMatchType::Type> OtherTagMatchType, EGameplayContainerMatchType ContainerMatchType) const
 {
-	UGameplayTagsManager& TagManager = IGameplayTagsModule::GetGameplayTagsManager();
+	UGameplayTagsManager& TagManager = UGameplayTagsManager::Get();
 
 	for (TArray<FGameplayTag>::TConstIterator OtherIt(OtherContainer.GameplayTags); OtherIt; ++OtherIt)
 	{
@@ -632,6 +668,7 @@ bool FGameplayTagContainer::DoesTagContainerMatchComplex(const FGameplayTagConta
 
 		for (TArray<FGameplayTag>::TConstIterator It(this->GameplayTags); It; ++It)
 		{
+			PRAGMA_DISABLE_DEPRECATION_WARNINGS
 			if (TagManager.GameplayTagsMatch(*It, TagMatchType, *OtherIt, OtherTagMatchType) == true)
 			{
 				if (ContainerMatchType == EGameplayContainerMatchType::Any)
@@ -644,6 +681,7 @@ bool FGameplayTagContainer::DoesTagContainerMatchComplex(const FGameplayTagConta
 				// we only need one match per tag in OtherContainer, so don't bother looking for more
 				break;
 			}
+			PRAGMA_ENABLE_DEPRECATION_WARNINGS
 		}
 
 		if (ContainerMatchType == EGameplayContainerMatchType::All && bTagFound == false)
@@ -668,12 +706,18 @@ void FGameplayTagContainer::AppendTags(FGameplayTagContainer const& Other)
 {
 	SCOPE_CYCLE_COUNTER(STAT_FGameplayTagContainer_AppendTags);
 
-	GameplayTags.Reserve(GameplayTags.Num() + Other.Num());
+	GameplayTags.Reserve(GameplayTags.Num() + Other.GameplayTags.Num());
+	ParentTags.Reserve(ParentTags.Num() + Other.ParentTags.Num());
+
+	// Add other container's tags to our own
+	for(const FGameplayTag& OtherTag : Other.GameplayTags)
+	{
+		GameplayTags.AddUnique(OtherTag);
+	}
 
-	//add all the tags
-	for(TArray<FGameplayTag>::TConstIterator It(Other.GameplayTags); It; ++It)
+	for (const FGameplayTag& OtherTag : Other.ParentTags)
 	{
-		AddTag(*It);
+		ParentTags.AddUnique(OtherTag);
 	}
 }
 
@@ -684,17 +728,19 @@ void FGameplayTagContainer::AppendMatchingTags(FGameplayTagContainer const& Othe
 {
 	SCOPE_CYCLE_COUNTER(STAT_FGameplayTagContainer_AppendMatchingTags);
 
-	for(TArray<FGameplayTag>::TConstIterator It(OtherA.GameplayTags); It; ++It)
+	for(const FGameplayTag& OtherATag : OtherA.GameplayTags)
 	{
-		if (OtherB.HasTag(*It, EGameplayTagMatchType::Explicit, EGameplayTagMatchType::IncludeParentTags))
+		if (OtherATag.MatchesAny(OtherB))
 		{
-			AddTag(*It);		
+			AddTag(OtherATag);
 		}
 	}
 }
 
 DECLARE_CYCLE_STAT(TEXT("FGameplayTagContainer::AddTag"), STAT_FGameplayTagContainer_AddTag, STATGROUP_GameplayTags);
 
+static UGameplayTagsManager* CachedTagManager = nullptr;
+
 
 void FGameplayTagContainer::AddTag(const FGameplayTag& TagToAdd)
 {
@@ -704,12 +750,53 @@ void FGameplayTagContainer::AddTag(const FGameplayTag& TagToAdd)
 	{
 		// Don't want duplicate tags
 		GameplayTags.AddUnique(TagToAdd);
+
+		AddParentsForTag(TagToAdd);
 	}
 }
 
 void FGameplayTagContainer::AddTagFast(const FGameplayTag& TagToAdd)
 {
 	GameplayTags.Add(TagToAdd);
+	AddParentsForTag(TagToAdd);
+}
+
+bool FGameplayTagContainer::AddLeafTag(const FGameplayTag& TagToAdd)
+{
+	UGameplayTagsManager& TagManager = UGameplayTagsManager::Get();
+
+	// Check tag is not already explicitly in container
+	if (HasTagExact(TagToAdd))
+	{
+		return true;
+	}
+
+	// If this tag is parent of explicitly added tag, fail
+	if (HasTag(TagToAdd))
+	{
+		return false;
+	}
+
+	const FGameplayTagContainer* TagToAddContainer = UGameplayTagsManager::Get().GetSingleTagContainer(TagToAdd);
+
+	// This should always succeed
+	if (!ensure(TagToAddContainer))
+	{
+		return false;
+	}
+
+	// Remove any tags in the container that are a parent to TagToAdd
+	for (const FGameplayTag& ParentTag : TagToAddContainer->ParentTags)
+	{
+		if (HasTagExact(ParentTag))
+		{
+			RemoveTag(ParentTag);
+		}
+	}
+
+	// Add the tag
+	AddTag(TagToAdd);
+	return true;
 }
 
 DECLARE_CYCLE_STAT(TEXT("FGameplayTagContainer::RemoveTag"), STAT_FGameplayTagContainer_RemoveTag, STATGROUP_GameplayTags);
@@ -718,7 +805,15 @@ bool FGameplayTagContainer::RemoveTag(FGameplayTag TagToRemove)
 {
 	SCOPE_CYCLE_COUNTER(STAT_FGameplayTagContainer_RemoveTag);
 
-	return GameplayTags.RemoveSingle(TagToRemove) > 0;
+	int32 NumChanged = GameplayTags.RemoveSingle(TagToRemove);
+
+	if (NumChanged > 0)
+	{
+		// Have to recompute parent table from scratch because there could be duplicates providing the same parent tag
+		FillParentTags();
+		return true;
+	}
+	return false;
 }
 
 DECLARE_CYCLE_STAT(TEXT("FGameplayTagContainer::RemoveTags"), STAT_FGameplayTagContainer_RemoveTags, STATGROUP_GameplayTags);
@@ -727,15 +822,26 @@ void FGameplayTagContainer::RemoveTags(FGameplayTagContainer TagsToRemove)
 {
 	SCOPE_CYCLE_COUNTER(STAT_FGameplayTagContainer_RemoveTags);
 
+	int32 NumChanged = 0;
+
 	for (auto Tag : TagsToRemove)
 	{
-		RemoveTag(Tag);
+		NumChanged += GameplayTags.RemoveSingle(Tag);
+	}
+
+	if (NumChanged > 0)
+	{
+		// Recompute once at the end
+		FillParentTags();
 	}
 }
 
-void FGameplayTagContainer::RemoveAllTags(int32)
+void FGameplayTagContainer::Reset(int32 Slack)
 {
-	GameplayTags.Reset();
+	GameplayTags.Reset(Slack);
+
+	// ParentTags is usually around size of GameplayTags on average
+	ParentTags.Reset(Slack);
 }
 
 bool FGameplayTagContainer::Serialize(FArchive& Ar)
@@ -744,7 +850,10 @@ bool FGameplayTagContainer::Serialize(FArchive& Ar)
 	
 	if (bOldTagVer)
 	{
+		TArray<FName> Tags_DEPRECATED;
 		Ar << Tags_DEPRECATED;
+		// Too old to deal with
+		UE_LOG(LogGameplayTags, Error, TEXT("Failed to load old GameplayTag container, too old to migrate correctly"));
 	}
 	else
 	{
@@ -752,68 +861,49 @@ bool FGameplayTagContainer::Serialize(FArchive& Ar)
 	}
 	
 	// Only do redirects for real loads, not for duplicates or recompiles
-	if (Ar.IsLoading() && Ar.IsPersistent() && !(Ar.GetPortFlags() & PPF_Duplicate) && !(Ar.GetPortFlags() & PPF_DuplicateForPIE))
+	if (Ar.IsLoading() )
 	{
-		UGameplayTagsManager& TagManager = IGameplayTagsModule::GetGameplayTagsManager();
-
-		// If loading old version, add old tags to the new gameplay tags array so they can be saved out with the new version
-		// This needs to happen 
-		// NOTE: DeprecatedTagNamesNotFoundInTagMap should be removed along with the bOldTagVer when we remove backwards
-		// compatibility, and the signature of RedirectTagsForContainer (below) should be changed appropriately as well.
-		TSet<FName> DeprecatedTagNamesNotFoundInTagMap;
-		if (bOldTagVer)
+		if (Ar.IsPersistent() && !(Ar.GetPortFlags() & PPF_Duplicate) && !(Ar.GetPortFlags() & PPF_DuplicateForPIE))
 		{
-			for (auto It = Tags_DEPRECATED.CreateConstIterator(); It; ++It)
-			{
-				const bool bErrorIfNotFound = false;
-				FGameplayTag Tag = TagManager.RequestGameplayTag(*It, bErrorIfNotFound);
-				if (Tag.IsValid())
-				{
-					TagManager.AddLeafTagToContainer(*this, Tag);
-				}
-				else
-				{	// For tags not found in the current table, add them to a list to be checked when handling
-					// redirection (below).
-					DeprecatedTagNamesNotFoundInTagMap.Add(*It);
-				}
-			}
+			// Rename any tags that may have changed by the ini file.  Redirects can happen regardless of version.
+			// Regardless of version, want loading to have a chance to handle redirects
+			UGameplayTagsManager::Get().RedirectTagsForContainer(*this, Ar.GetSerializedProperty());
 		}
 
-		// Rename any tags that may have changed by the ini file.  Redirects can happen regardless of version.
-		// Regardless of version, want loading to have a chance to handle redirects
-		TagManager.RedirectTagsForContainer(*this, DeprecatedTagNamesNotFoundInTagMap, Ar.GetSerializedProperty());
+		FillParentTags();
+	}
+
+	if (Ar.IsSaving())
+	{
+		// This marks the saved name for later searching
+		for (const FGameplayTag& Tag : GameplayTags)
+		{
+			Ar.MarkSearchableName(FGameplayTag::StaticStruct(), Tag.TagName);
+		}
 	}
 
 	return true;
 }
 
-int32 FGameplayTagContainer::Num() const
+FString FGameplayTagContainer::ToString() const
 {
-	return GameplayTags.Num();
+	FString ExportString;
+	FGameplayTagContainer::StaticStruct()->ExportText(ExportString, this, this, nullptr, 0, nullptr);
+
+	return ExportString;
 }
 
-FString FGameplayTagContainer::ToString() const
+bool FGameplayTagContainer::ImportTextItem(const TCHAR*& Buffer, int32 PortFlags, UObject* Parent, FOutputDevice* ErrorText)
 {
-	FString RetString = TEXT("(GameplayTags=");
-	if (GameplayTags.Num() > 0)
-	{
-		RetString += TEXT("(");
+	// Call default import, but skip the native callback to avoid recursion
+	Buffer = FGameplayTagContainer::StaticStruct()->ImportText(Buffer, this, Parent, PortFlags, ErrorText, TEXT("FGameplayTagContainer"), false);
 
-		for (int i = 0; i < GameplayTags.Num(); ++i)
-		{
-			RetString += TEXT("(TagName=\"");
-			RetString += GameplayTags[i].ToString();
-			RetString += TEXT("\")");
-			if (i < GameplayTags.Num() - 1)
-			{
-				RetString += TEXT(",");
-			}
-		}
-		RetString += TEXT(")");
+	if (Buffer)
+	{
+		// Compute parent tags
+		FillParentTags();	
 	}
-	RetString += TEXT(")");
-
-	return RetString;
+	return true;
 }
 
 FString FGameplayTagContainer::ToStringSimple() const
@@ -841,7 +931,7 @@ bool FGameplayTagContainer::NetSerialize(FArchive& Ar, class UPackageMap* Map, b
 	{
 		if (GameplayTags.Num() > 0)
 		{
-			GameplayTags.Reset();
+			Reset();
 		}
 		bOutSuccess = true;
 		return true;
@@ -849,23 +939,25 @@ bool FGameplayTagContainer::NetSerialize(FArchive& Ar, class UPackageMap* Map, b
 
 	// -------------------------------------------------------
 
+	int32 NumBitsForContainerSize = UGameplayTagsManager::Get().NumBitsForContainerSize;
+
 	if (Ar.IsSaving())
 	{
 		uint8 NumTags = GameplayTags.Num();
-		uint8 MaxSize = (1 << UGameplayTagsManager::NumBitsForContainerSize);
+		uint8 MaxSize = (1 << NumBitsForContainerSize);
 		if (!ensureMsgf(NumTags < MaxSize, TEXT("TagContainer has %d elements when max is %d! Tags: %s"), NumTags, MaxSize, *ToStringSimple()))
 		{
 			NumTags = MaxSize - 1;
 		}
 		
-		Ar.SerializeBits(&NumTags, UGameplayTagsManager::NumBitsForContainerSize);
+		Ar.SerializeBits(&NumTags, NumBitsForContainerSize);
 		for (int32 idx=0; idx < NumTags;++idx)
 		{
 			FGameplayTag& Tag = GameplayTags[idx];
 			Tag.NetSerialize_Packed(Ar, Map, bOutSuccess);
 
 #if !(UE_BUILD_SHIPPING || UE_BUILD_TEST)
-			UGameplayTagsManager::NotifyTagReplicated(Tag, true);
+			UGameplayTagsManager::Get().NotifyTagReplicated(Tag, true);
 #endif
 		}
 	}
@@ -873,7 +965,7 @@ bool FGameplayTagContainer::NetSerialize(FArchive& Ar, class UPackageMap* Map, b
 	{
 		// No Common Container tags, just replicate this like normal
 		uint8 NumTags = 0;
-		Ar.SerializeBits(&NumTags, UGameplayTagsManager::NumBitsForContainerSize);
+		Ar.SerializeBits(&NumTags, NumBitsForContainerSize);
 
 		GameplayTags.Empty(NumTags);
 		GameplayTags.AddDefaulted(NumTags);
@@ -881,6 +973,7 @@ bool FGameplayTagContainer::NetSerialize(FArchive& Ar, class UPackageMap* Map, b
 		{
 			GameplayTags[idx].NetSerialize_Packed(Ar, Map, bOutSuccess);
 		}
+		FillParentTags();
 	}
 
 
@@ -928,20 +1021,87 @@ FText FGameplayTagContainer::ToMatchingText(EGameplayContainerMatchType MatchTyp
 
 bool FGameplayTag::ComplexMatches(TEnumAsByte<EGameplayTagMatchType::Type> MatchTypeOne, const FGameplayTag& Other, TEnumAsByte<EGameplayTagMatchType::Type> MatchTypeTwo) const
 {
-	return IGameplayTagsModule::Get().GetGameplayTagsManager().ComplexGameplayTagsMatch(*this, MatchTypeOne, Other, MatchTypeTwo);
+	PRAGMA_DISABLE_DEPRECATION_WARNINGS
+	return UGameplayTagsManager::Get().GameplayTagsMatch(*this, MatchTypeOne, Other, MatchTypeTwo);
+	PRAGMA_ENABLE_DEPRECATION_WARNINGS
 }
 
-#if CHECK_TAG_OPTIMIZATIONS
-bool FGameplayTag::MatchesOriginal(TEnumAsByte<EGameplayTagMatchType::Type> MatchTypeOne, const FGameplayTag& Other, TEnumAsByte<EGameplayTagMatchType::Type> MatchTypeTwo) const
+DECLARE_CYCLE_STAT(TEXT("FGameplayTag::GetSingleTagContainer"), STAT_FGameplayTag_GetSingleTagContainer, STATGROUP_GameplayTags);
+
+const FGameplayTagContainer& FGameplayTag::GetSingleTagContainer() const
 {
-	return IGameplayTagsModule::Get().GetGameplayTagsManager().GameplayTagsMatchOriginal(*this, MatchTypeOne, Other, MatchTypeTwo);
+	SCOPE_CYCLE_COUNTER(STAT_FGameplayTag_GetSingleTagContainer);
+
+	const FGameplayTagContainer* TagContainer = UGameplayTagsManager::Get().GetSingleTagContainer(*this);
+
+	if (TagContainer)
+	{
+		return *TagContainer;
+	}
+
+	// This should always be invalid if the node is missing
+	ensure(!IsValid());
+
+	return FGameplayTagContainer::EmptyContainer;
+}
+
+FGameplayTag FGameplayTag::RequestGameplayTag(FName TagName, bool ErrorIfNotFound)
+{
+	return UGameplayTagsManager::Get().RequestGameplayTag(TagName, ErrorIfNotFound);
+}
+
+FGameplayTagContainer FGameplayTag::GetGameplayTagParents() const
+{
+	return UGameplayTagsManager::Get().RequestGameplayTagParents(*this);
+}
+
+DECLARE_CYCLE_STAT(TEXT("FGameplayTag::MatchesTag"), STAT_FGameplayTag_MatchesTag, STATGROUP_GameplayTags);
+
+bool FGameplayTag::MatchesTag(const FGameplayTag& TagToCheck) const
+{
+	SCOPE_CYCLE_COUNTER(STAT_FGameplayTag_MatchesTag);
+
+	const FGameplayTagContainer* TagContainer = UGameplayTagsManager::Get().GetSingleTagContainer(*this);
+
+	if (TagContainer)
+	{
+		return TagContainer->HasTag(TagToCheck);
+	}
+
+	// This should always be invalid if the node is missing
+	ensure(!IsValid());
+
+	return false;
+}
+
+DECLARE_CYCLE_STAT(TEXT("FGameplayTag::MatchesAny"), STAT_FGameplayTag_MatchesAny, STATGROUP_GameplayTags);
+
+bool FGameplayTag::MatchesAny(const FGameplayTagContainer& ContainerToCheck) const
+{
+	SCOPE_CYCLE_COUNTER(STAT_FGameplayTag_MatchesAny);
+
+	const FGameplayTagContainer* TagContainer = UGameplayTagsManager::Get().GetSingleTagContainer(*this);
+
+	if (TagContainer)
+	{
+		return TagContainer->HasAny(ContainerToCheck);
+	}
+
+	// This should always be invalid if the node is missing
+	ensure(!IsValid());
+
+	return false;
+}
+
+int32 FGameplayTag::MatchesTagDepth(const FGameplayTag& TagToCheck) const
+{
+	return UGameplayTagsManager::Get().GameplayTagsMatchDepth(*this, TagToCheck);
 }
-#endif
 
 FGameplayTag::FGameplayTag(FName Name)
 	: TagName(Name)
 {
-	check(IGameplayTagsModule::Get().GetGameplayTagsManager().ValidateTagCreation(Name));
+	// This constructor is used to bypass the table check and is only usable by GameplayTagManager
 }
 
 bool FGameplayTag::SerializeFromMismatchedTag(const FPropertyTag& Tag, FArchive& Ar)
@@ -956,7 +1116,7 @@ bool FGameplayTag::SerializeFromMismatchedTag(const FPropertyTag& Tag, FArchive&
 
 FGameplayTag FGameplayTag::RequestDirectParent() const
 {
-	return IGameplayTagsModule::Get().GetGameplayTagsManager().RequestGameplayTagDirectParent(*this);
+	return UGameplayTagsManager::Get().RequestGameplayTagDirectParent(*this);
 }
 
 DECLARE_CYCLE_STAT(TEXT("FGameplayTag::NetSerialize"), STAT_FGameplayTag_NetSerialize, STATGROUP_GameplayTags);
@@ -966,7 +1126,7 @@ bool FGameplayTag::NetSerialize(FArchive& Ar, class UPackageMap* Map, bool& bOut
 #if !(UE_BUILD_SHIPPING || UE_BUILD_TEST)
 	if (Ar.IsSaving())
 	{
-		UGameplayTagsManager::NotifyTagReplicated(*this, false);
+		UGameplayTagsManager::Get().NotifyTagReplicated(*this, false);
 	}
 #endif
 
@@ -990,7 +1150,7 @@ static TSharedPtr<FNetFieldExportGroup> CreateNetfieldExportGroupForNetworkGamep
 		FNetFieldExport NetFieldExport(
 			i,
 			0,
-			NetworkGameplayTagNodeIndex[i]->GetCompleteTag().ToString(),
+			NetworkGameplayTagNodeIndex[i]->GetCompleteTagString(),
 			TEXT(""));
 
 		NetFieldExportGroup->NetFieldExports[i] = NetFieldExport;
@@ -1003,7 +1163,7 @@ bool FGameplayTag::NetSerialize_Packed(FArchive& Ar, class UPackageMap* Map, boo
 {
 	SCOPE_CYCLE_COUNTER(STAT_FGameplayTag_NetSerialize);
 
-	UGameplayTagsManager& TagManager = IGameplayTagsModule::GetGameplayTagsManager();	
+	UGameplayTagsManager& TagManager = UGameplayTagsManager::Get();
 
 	if (TagManager.ShouldUseFastReplication())
 	{
@@ -1056,7 +1216,7 @@ bool FGameplayTag::NetSerialize_Packed(FArchive& Ar, class UPackageMap* Map, boo
 					TagName = FName(*NetFieldExportGroup->NetFieldExports[NetIndex].Name);
 
 					// Validate the tag name
-					const FGameplayTag Tag = IGameplayTagsModule::RequestGameplayTag(TagName, false);
+					const FGameplayTag Tag = UGameplayTagsManager::Get().RequestGameplayTag(TagName, false);
 
 					// Warn (once) if the tag isn't found
 					if (!Tag.IsValid() && !NetFieldExportGroup->NetFieldExports[NetIndex].bIncompatible)
@@ -1081,11 +1241,11 @@ bool FGameplayTag::NetSerialize_Packed(FArchive& Ar, class UPackageMap* Map, boo
 		{
 			NetIndex = TagManager.GetNetIndexFromTag(*this);
 			
-			SerializeTagNetIndexPacked(Ar, NetIndex, UGameplayTagsManager::NetIndexFirstBitSegment, UGameplayTagsManager::NetIndexTrueBitNum);
+			SerializeTagNetIndexPacked(Ar, NetIndex, TagManager.NetIndexFirstBitSegment, TagManager.NetIndexTrueBitNum);
 		}
 		else
 		{
-			SerializeTagNetIndexPacked(Ar, NetIndex, UGameplayTagsManager::NetIndexFirstBitSegment, UGameplayTagsManager::NetIndexTrueBitNum);
+			SerializeTagNetIndexPacked(Ar, NetIndex, TagManager.NetIndexFirstBitSegment, TagManager.NetIndexTrueBitNum);
 			TagName = TagManager.GetTagNameFromNetIndex(NetIndex);
 		}
 	}
@@ -1104,10 +1264,14 @@ void FGameplayTag::PostSerialize(const FArchive& Ar)
 	// Only do redirects for real loads, not for duplicates or recompiles
 	if (Ar.IsLoading() && Ar.IsPersistent() && !(Ar.GetPortFlags() & PPF_Duplicate) && !(Ar.GetPortFlags() & PPF_DuplicateForPIE))
 	{
-		UGameplayTagsManager& TagManager = IGameplayTagsModule::GetGameplayTagsManager();
-
 		// Rename any tags that may have changed by the ini file.
-		TagManager.RedirectSingleGameplayTag(*this, Ar.GetSerializedProperty());
+		UGameplayTagsManager::Get().RedirectSingleGameplayTag(*this, Ar.GetSerializedProperty());
+	}
+
+	if (Ar.IsSaving() && IsValid())
+	{
+		// This marks the saved name for later searching
+		Ar.MarkSearchableName(FGameplayTag::StaticStruct(), TagName);
 	}
 }
 
@@ -1135,7 +1299,7 @@ bool FGameplayTag::ImportTextItem(const TCHAR*& Buffer, int32 PortFlags, UObject
 	}
 
 	FName ImportedTagName = FName(*ImportedTag);
-	if (IGameplayTagsModule::Get().GetGameplayTagsManager().ValidateTagCreation(ImportedTagName))
+	if (UGameplayTagsManager::Get().ValidateTagCreation(ImportedTagName))
 	{
 		// We found the tag. Assign it here.
 		TagName = ImportedTagName;
@@ -1640,7 +1804,7 @@ void UEditableGameplayTagQueryExpression_NoExprMatch::EmitTokens(TArray<uint8>&
 
 FGameplayTagQueryExpression& FGameplayTagQueryExpression::AddTag(FName TagName)
 {
-	FGameplayTag const Tag = IGameplayTagsModule::RequestGameplayTag(TagName);
+	FGameplayTag const Tag = UGameplayTagsManager::Get().RequestGameplayTag(TagName);
 	return AddTag(Tag);
 }
 
@@ -1692,7 +1856,7 @@ void FGameplayTagQueryExpression::EmitTokens(TArray<uint8>& TokenStream, TArray<
 
 static void GameplayTagPrintReplicationMap()
 {
-	UGameplayTagsManager::PrintReplicationFrequencyReport();
+	UGameplayTagsManager::Get().PrintReplicationFrequencyReport();
 }
 
 FAutoConsoleCommand GameplayTagPrintReplicationMapCmd(
diff --git a/Engine/Source/Runtime/GameplayTags/Private/GameplayTagsManager.cpp b/Engine/Source/Runtime/GameplayTags/Private/GameplayTagsManager.cpp
--- a/Engine/Source/Runtime/GameplayTags/Private/GameplayTagsManager.cpp
+++ b/Engine/Source/Runtime/GameplayTags/Private/GameplayTagsManager.cpp
@@ -2,123 +2,51 @@
 
 #include "GameplayTagsModulePrivatePCH.h"
 #include "GameplayTagsSettings.h"
-#include "SNotificationList.h"
-#include "NotificationManager.h"
-
-#if WITH_EDITOR
-#include "UnrealEd.h"
-#endif
 
 #define LOCTEXT_NAMESPACE "GameplayTagManager"
 
-#if !(UE_BUILD_SHIPPING || UE_BUILD_TEST)
-TMap<FGameplayTag, int32>	UGameplayTagsManager::ReplicationCountMap;
-TMap<FGameplayTag, int32>	UGameplayTagsManager::ReplicationCountMap_SingleTags;
-TMap<FGameplayTag, int32>	UGameplayTagsManager::ReplicationCountMap_Containers;
-#endif
-
-int32 UGameplayTagsManager::NetIndexFirstBitSegment=16;
-int32 UGameplayTagsManager::NetIndexTrueBitNum=16;
-int32 UGameplayTagsManager::NumBitsForContainerSize=6;
+UGameplayTagsManager* UGameplayTagsManager::SingletonManager = nullptr;
 
 UGameplayTagsManager::UGameplayTagsManager(const FObjectInitializer& ObjectInitializer)
 	: Super(ObjectInitializer)
 {
-#if WITH_EDITOR
-	RegisteredObjectReimport = false;
-#endif
+	bUseFastReplication = false;
+	bShouldWarnOnInvalidTags = true;
+	NetIndexFirstBitSegment = 16;
+	NetIndexTrueBitNum = 16;
+	NumBitsForContainerSize = 6;
 }
 
-void UGameplayTagsManager::LoadGameplayTagTables(TArray<FString>& TagTableNames)
+void UGameplayTagsManager::LoadGameplayTagTables()
 {
-	if (GameplayTagTables.Num() == 0 && TagTableNames.Num() > 0)
-	{
-		for (auto It(TagTableNames.CreateConstIterator()); It; ++It)
-		{
-			const FString& FileName = *It;
-			UDataTable* TagTable = LoadObject<UDataTable>(NULL, *FileName, NULL, LOAD_None, NULL);
+	GameplayTagTables.Empty();
 
-			// Handle case where the module is dynamically-loaded within a LoadPackage stack, which would otherwise
-			// result in the tag table not having its RowStruct serialized in time. Without the RowStruct, the tags manager
-			// will not be initialized correctly.
-			if (TagTable && IsLoading())
-			{
-				FLinkerLoad* TagLinker = TagTable->GetLinker();
-				if (TagLinker)
-				{
-					TagTable->GetLinker()->Preload(TagTable);
-				}
-			}
-			GameplayTagTables.Add(TagTable);
-		}
-	}
+	UGameplayTagsSettings* MutableDefault = GetMutableDefault<UGameplayTagsSettings>();
 
-#if WITH_EDITOR
-	// Hook into notifications for object re-imports so that the gameplay tag tree can be reconstructed if the table changes
-	if (GIsEditor && GameplayTagTables.Num() > 0 && !RegisteredObjectReimport)
+	for (FStringAssetReference DataTablePath : MutableDefault->GameplayTagTableList)
 	{
-		RegisteredObjectReimport = true;
-		FEditorDelegates::OnAssetPostImport.AddUObject(this, &UGameplayTagsManager::OnObjectReimported);
-	}
-#endif
-}
+		UDataTable* TagTable = LoadObject<UDataTable>(nullptr, *DataTablePath.ToString(), nullptr, LOAD_None, nullptr);
 
-void UGameplayTagsManager::GetAllNodesForTag_Recurse(TArray<FString>& Tags, int32 CurrentTagDepth, TSharedPtr<FGameplayTagNode> CurrentTagNode, TArray< TSharedPtr<FGameplayTagNode> >& OutTagArray) const
-{
-	CurrentTagDepth++;
-	if(Tags.Num() > CurrentTagDepth)
-	{
-		// search for the subsequent tags in the hierarchy
-		TArray< TSharedPtr<FGameplayTagNode> > CurrentChildrenTags;
-		CurrentChildrenTags = CurrentTagNode->GetChildTagNodes();
-		for(int32 TagIdx = 0; TagIdx < CurrentChildrenTags.Num(); TagIdx++)
+		// Handle case where the module is dynamically-loaded within a LoadPackage stack, which would otherwise
+		// result in the tag table not having its RowStruct serialized in time. Without the RowStruct, the tags manager
+		// will not be initialized correctly.
+		if (TagTable && IsLoading())
 		{
-			FString CurrentTagName = CurrentChildrenTags[TagIdx].Get()->GetSimpleTag().ToString();
-			if(CurrentTagName == Tags[CurrentTagDepth])
+			FLinkerLoad* TagLinker = TagTable->GetLinker();
+			if (TagLinker)
 			{
-				CurrentTagNode = CurrentChildrenTags[TagIdx];
-				OutTagArray.Add(CurrentTagNode);
-				GetAllNodesForTag_Recurse(Tags, CurrentTagDepth, CurrentTagNode, OutTagArray);
-				break;
+				TagTable->GetLinker()->Preload(TagTable);
 			}
 		}
+		GameplayTagTables.Add(TagTable);
 	}
 }
 
-void UGameplayTagsManager::GetAllNodesForTag( const FString& Tag, TArray< TSharedPtr<FGameplayTagNode> >& OutTagArray ) const
-{
-	TArray<FString> Tags;
-	TArray<TSharedPtr<FGameplayTagNode>>& GameplayRootTags = GameplayRootTag->GetChildTagNodes();
-	OutTagArray.Empty();
-	if(Tag.ParseIntoArray( Tags, TEXT( "." ), true ) > 0)
-	{
-		int32 CurrentTagDepth = 0;
-		// find the first node of the tag
-		TSharedPtr<FGameplayTagNode> CurrentTagNode = NULL;
-		for(int32 TagIdx = 0; TagIdx < GameplayRootTags.Num(); TagIdx++)
-		{
-			FString TagName = GameplayRootTags[TagIdx].Get()->GetSimpleTag().ToString();
-			if(TagName == Tags[CurrentTagDepth])
-			{
-				CurrentTagNode = GameplayRootTags[TagIdx];
-				break;
-			}
-		}
-
-		if(CurrentTagNode.IsValid())
-		{
-			// add it to the list of tags
-			OutTagArray.Add(CurrentTagNode);
-			GetAllNodesForTag_Recurse(Tags, CurrentTagDepth, CurrentTagNode, OutTagArray);
-		}
-	}	
-}
-
 struct FCompareFGameplayTagNodeByTag
 {
 	FORCEINLINE bool operator()( const TSharedPtr<FGameplayTagNode>& A, const TSharedPtr<FGameplayTagNode>& B ) const
 	{
-		return (A->GetSimpleTag().Compare(B->GetSimpleTag())) < 0;
+		return (A->GetSimpleTagName().Compare(B->GetSimpleTagName())) < 0;
 	}
 };
 
@@ -142,96 +70,105 @@ void UGameplayTagsManager::ConstructGameplayTagTree()
 			}
 		}
 
+		UGameplayTagsSettings* MutableDefault = GetMutableDefault<UGameplayTagsSettings>();
+		FString DefaultEnginePath = FString::Printf(TEXT("%sDefaultEngine.ini"), *FPaths::SourceConfigDir());
+
+		// Create native source
+		FindOrAddTagSource(FGameplayTagSource::GetNativeName(), EGameplayTagSourceType::Native);
+
 		if (ShouldImportTagsFromINI())
 		{
-			UGameplayTagsSettings* MutableDefault = GetMutableDefault<UGameplayTagsSettings>();
 #if STATS
 			FString PerfMessage = FString::Printf(TEXT("UGameplayTagsManager::ConstructGameplayTagTree: ImportINI"));
 			SCOPE_LOG_TIME_IN_SECONDS(*PerfMessage, nullptr)
 #endif
-			// ------------------------------------------------------------------------------------------------------------
-			// Update path: Check for old tags in DefaultEngine.ini (we'll push them to the UGameplayTagSettings class).
-			// ------------------------------------------------------------------------------------------------------------
+			// Copy from deprecated list in DefaultEngine.ini
 			TArray<FString> EngineConfigTags;
-			GConfig->GetArray(TEXT("/Script/GameplayTags.GameplayTagsSettings"), TEXT("GameplayTags"), EngineConfigTags, GEngineIni);
+			GConfig->GetArray(TEXT("/Script/GameplayTags.GameplayTagsSettings"), TEXT("+GameplayTags"), EngineConfigTags, DefaultEnginePath);
 			
-			if (EngineConfigTags.Num() > 0)
+			for (const FString& EngineConfigTag : EngineConfigTags)
 			{
-				if (MutableDefault->GameplayTags.Num() == 0)
-				{
-					MutableDefault->GameplayTags.Append(EngineConfigTags);
-				}
+				MutableDefault->GameplayTagList.AddUnique(FGameplayTagTableRow(FName(*EngineConfigTag)));
 			}
 
-			// ---------------------------------------------------------------
-			//	Developer Tags
-			// ---------------------------------------------------------------
+			// Copy from deprecated list in DefaultGamplayTags.ini
+			EngineConfigTags.Empty();
+			GConfig->GetArray(TEXT("/Script/GameplayTags.GameplayTagsSettings"), TEXT("+GameplayTags"), EngineConfigTags, MutableDefault->GetDefaultConfigFilename());
 
+			for (const FString& EngineConfigTag : EngineConfigTags)
 			{
-				// Read all tags from the ini
-				DeveloperTags.Items.Reset();
-				TArray<FString> FilesInDirectory;
-				IFileManager::Get().FindFilesRecursive(FilesInDirectory, *(FPaths::GameConfigDir() / TEXT("Tags")), TEXT("*.ini"), true, false);
-				FilesInDirectory.Sort();
-				for (FString& FileName : FilesInDirectory)
-				{
-					UE_LOG(LogGameplayTags, Display, TEXT("File: %s"), *FileName);
+				MutableDefault->GameplayTagList.AddUnique(FGameplayTagTableRow(FName(*EngineConfigTag)));
+			}
 
-					auto& NewItem = DeveloperTags.Items[DeveloperTags.Items.AddDefaulted()];
-					NewItem.IniName = FileName;
+#if WITH_EDITOR
+			MutableDefault->SortTags();
+#endif
 
-					GConfig->GetArray(TEXT("UserTags"), TEXT("GameplayTags"), NewItem.Tags, FileName);
+			FName TagSource = FGameplayTagSource::GetDefaultName();
+			FGameplayTagSource* DefaultSource = FindOrAddTagSource(TagSource, EGameplayTagSourceType::DefaultTagList);
 
-					MutableDefault->GameplayTags.Append(NewItem.Tags);
-				}
+			for (const FGameplayTagTableRow& TableRow : MutableDefault->GameplayTagList)
+			{
+				AddTagTableRow(TableRow, TagSource);
 			}
-			
-
-			// ---------------------------------------------------------------
 
-			// Load any GameplayTagSettings from config (their default object)
-			for (TObjectIterator<UClass> ClassIt; ClassIt; ++ClassIt)
+			// Extra tags
+		
+			// Read all tags from the ini
+			TArray<FString> FilesInDirectory;
+			IFileManager::Get().FindFilesRecursive(FilesInDirectory, *(FPaths::GameConfigDir() / TEXT("Tags")), TEXT("*.ini"), true, false);
+			FilesInDirectory.Sort();
+			for (FString& FileName : FilesInDirectory)
 			{
-				UClass* Class = *ClassIt;
-				if (!Class->IsChildOf<UGameplayTagsSettings>() || Class->HasAnyClassFlags(CLASS_Abstract))
-				{
-					continue;
-				}
+				TagSource = FName(*FPaths::GetCleanFilename(FileName));
+				FGameplayTagSource* FoundSource = FindOrAddTagSource(TagSource, EGameplayTagSourceType::TagList);
 
-#if WITH_EDITOR
-				Class->GetDefaultObject<UGameplayTagsSettings>()->SortTags();
-#endif
-				for (FString& TagStr : Class->GetDefaultObject<UGameplayTagsSettings>()->GameplayTags)
-				{
-					FGameplayTagTableRow TableRow;
-					TableRow.Tag = TagStr;
-					AddTagTableRow(TableRow);
-				}
+				UE_LOG(LogGameplayTags, Display, TEXT("Loading Tag File: %s"), *FileName);
 
-				// Grab the commonly replicated tags
-				CommonlyReplicatedTags.Reserve(CommonlyReplicatedTags.Num() + Class->GetDefaultObject<UGameplayTagsSettings>()->CommonlyReplicatedTags.Num());
-				for (const FString& Str : Class->GetDefaultObject<UGameplayTagsSettings>()->CommonlyReplicatedTags)
+				// Check deprecated locations
+				TArray<FString> Tags;
+				if (GConfig->GetArray(TEXT("UserTags"), TEXT("GameplayTags"), Tags, FileName))
 				{
-					FGameplayTag Tag = RequestGameplayTag(FName(*Str));
-					if (Tag.IsValid())
-					{
-						CommonlyReplicatedTags.Add(Tag);
-					}
-					else
+					for (const FString& Tag : Tags)
 					{
-						UE_LOG(LogGameplayTags, Warning, TEXT("%s was found in the CommonlyReplicatedTags list but doesn't appear to be a valid tag!"), *Str);
+						FoundSource->SourceTagList->GameplayTagList.AddUnique(FGameplayTagTableRow(FName(*Tag)));
 					}
 				}
+				else
+				{
+					// Load from new ini
+					FoundSource->SourceTagList->LoadConfig(UGameplayTagsList::StaticClass(), *FileName);
+				}
 
-				NetIndexFirstBitSegment = Class->GetDefaultObject<UGameplayTagsSettings>()->NetIndexFirstBitSegment;
+				// Add to dictionary
+				for (const FGameplayTagTableRow& TableRow : FoundSource->SourceTagList->GameplayTagList)
+				{
+					AddTagTableRow(TableRow, TagSource);
+				}
+			}
+		
+			// Grab the commonly replicated tags
+			CommonlyReplicatedTags.Empty();
+			for (FName TagName : MutableDefault->CommonlyReplicatedTags)
+			{
+				FGameplayTag Tag = RequestGameplayTag(TagName);
+				if (Tag.IsValid())
+				{
+					CommonlyReplicatedTags.Add(Tag);
+				}
+				else
+				{
+					UE_LOG(LogGameplayTags, Warning, TEXT("%s was found in the CommonlyReplicatedTags list but doesn't appear to be a valid tag!"), *TagName.ToString());
+				}
 			}
+
 			GameplayRootTag->GetChildTagNodes().Sort(FCompareFGameplayTagNodeByTag());
 		}
 
-		bUseFastReplication = false;
-		GConfig->GetBool(TEXT("GameplayTags"), TEXT("FastReplication"), bUseFastReplication, GEngineIni);
-
-		GConfig->GetInt(TEXT("GameplayTags"), TEXT("NumBitsForContainerSize"), NumBitsForContainerSize, GEngineIni);
+		bUseFastReplication = MutableDefault->FastReplication;
+		bShouldWarnOnInvalidTags = MutableDefault->WarnOnInvalidTags;
+		NumBitsForContainerSize = MutableDefault->NumBitsForContainerSize;
+		NetIndexFirstBitSegment = MutableDefault->NetIndexFirstBitSegment;
 
 		if (ShouldUseFastReplication())
 		{
@@ -247,51 +184,78 @@ void UGameplayTagsManager::ConstructGameplayTagTree()
 			FString PerfMessage = FString::Printf(TEXT("UGameplayTagsManager::ConstructGameplayTagTree: GameplayTagTreeChangedEvent.Broadcast"));
 			SCOPE_LOG_TIME_IN_SECONDS(*PerfMessage, nullptr)
 #endif
-			GameplayTagTreeChangedEvent.Broadcast();
+			IGameplayTagsModule::OnGameplayTagTreeChanged.Broadcast();
 		}
 
 		// Update the TagRedirects map
 		TagRedirects.Empty();
-		FConfigSection* PackageRedirects = GConfig->GetSectionPrivate(TEXT("/Script/Engine.Engine"), false, true, GEngineIni);
-		for (FConfigSection::TIterator It(*PackageRedirects); It; ++It)
+
+		// Check the deprecated location
+		bool bFoundDeprecated = false;
+		FConfigSection* PackageRedirects = GConfig->GetSectionPrivate(TEXT("/Script/Engine.Engine"), false, true, DefaultEnginePath);
+
+		if (PackageRedirects)
 		{
-			if (It.Key() == TEXT("GameplayTagRedirects"))
+			for (FConfigSection::TIterator It(*PackageRedirects); It; ++It)
 			{
-				FName OldTagName = NAME_None;
-				FName NewTagName;
-
-				if (FParse::Value(*It.Value().GetValue(), TEXT("OldTagName="), OldTagName))
+				if (It.Key() == TEXT("+GameplayTagRedirects"))
 				{
-					if (FParse::Value(*It.Value().GetValue(), TEXT("NewTagName="), NewTagName))
+					FName OldTagName = NAME_None;
+					FName NewTagName;
+
+					if (FParse::Value(*It.Value().GetValue(), TEXT("OldTagName="), OldTagName))
 					{
-						if (ensureMsgf(!TagRedirects.Contains(OldTagName), TEXT("Old tag %s is being redirected to more than one tag. Please remove all the redirections except for one."), *OldTagName.ToString()))
+						if (FParse::Value(*It.Value().GetValue(), TEXT("NewTagName="), NewTagName))
 						{
-							FGameplayTag OldTag = RequestGameplayTag(OldTagName, false); //< This only succeeds if OldTag is in the Table!
-							if (OldTag.IsValid())
-							{
-								UE_LOG(LogGameplayTags, Warning,
-									TEXT("Old tag (%s) which is being redirected still exists in the table!  Generally you should "
-									TEXT("remove the old tags from the table when you are redirecting to new tags, or else users will ")
-									TEXT("still be able to add the old tags to containers.")), *OldTagName.ToString()
-									);
-							}
-
-							FGameplayTag NewTag = (NewTagName != NAME_None) ? RequestGameplayTag(NewTagName, false) : FGameplayTag();
-							if (!NewTag.IsValid() && NewTagName != NAME_None)
-							{
-								UE_LOG(LogGameplayTags, Warning, TEXT("Invalid new tag %s!  Cannot replace old tag %s."),
-									*NewTagName.ToString(), *OldTagName.ToString());
-							}
-							else
-							{
-								// Populate the map
-								TagRedirects.Add(OldTagName, NewTag);
-							}
+							FGameplayTagRedirect Redirect;
+							Redirect.OldTagName = OldTagName;
+							Redirect.NewTagName = NewTagName;
+
+							MutableDefault->GameplayTagRedirects.AddUnique(Redirect);
+
+							bFoundDeprecated = true;
 						}
 					}
 				}
 			}
 		}
+
+		if (bFoundDeprecated)
+		{
+			UE_LOG(LogGameplayTags, Log, TEXT("GameplayTagRedirects is in a deprecated location, after editing GameplayTags developer settings you must remove these manually"));
+		}
+
+		// Check settings object
+		for (const FGameplayTagRedirect& Redirect : MutableDefault->GameplayTagRedirects)
+		{
+			FName OldTagName = Redirect.OldTagName;
+			FName NewTagName = Redirect.NewTagName;
+
+			if (ensureMsgf(!TagRedirects.Contains(OldTagName), TEXT("Old tag %s is being redirected to more than one tag. Please remove all the redirections except for one."), *OldTagName.ToString()))
+			{
+				FGameplayTag OldTag = RequestGameplayTag(OldTagName, false); //< This only succeeds if OldTag is in the Table!
+				if (OldTag.IsValid())
+				{
+					UE_LOG(LogGameplayTags, Warning,
+						TEXT("Old tag (%s) which is being redirected still exists in the table!  Generally you should "
+						TEXT("remove the old tags from the table when you are redirecting to new tags, or else users will ")
+						TEXT("still be able to add the old tags to containers.")), *OldTagName.ToString()
+						);
+				}
+
+				FGameplayTag NewTag = (NewTagName != NAME_None) ? RequestGameplayTag(NewTagName, false) : FGameplayTag();
+				if (!NewTag.IsValid() && NewTagName != NAME_None)
+				{
+					UE_LOG(LogGameplayTags, Warning, TEXT("Invalid new tag %s!  Cannot replace old tag %s."),
+						*NewTagName.ToString(), *OldTagName.ToString());
+				}
+				else
+				{
+					// Populate the map
+					TagRedirects.Add(OldTagName, NewTag);
+				}
+			}
+		}
 	}
 }
 
@@ -303,6 +267,8 @@ void UGameplayTagsManager::ConstructNetIndex()
 
 	NetworkGameplayTagNodeIndex.Sort(FCompareFGameplayTagNodeByTag());
 
+	check(CommonlyReplicatedTags.Num() < NetworkGameplayTagNodeIndex.Num());
+
 	// Put the common indices up front
 	for (int32 CommonIdx=0; CommonIdx < CommonlyReplicatedTags.Num(); ++CommonIdx)
 	{
@@ -312,7 +278,7 @@ void UGameplayTagsManager::ConstructNetIndex()
 		bool Found = false;
 		for (int32 findidx=0; findidx < NetworkGameplayTagNodeIndex.Num(); ++findidx)
 		{
-			if (NetworkGameplayTagNodeIndex[findidx]->GetCompleteTag() == Tag.GetTagName())
+			if (NetworkGameplayTagNodeIndex[findidx]->GetCompleteTag() == Tag)
 			{
 				NetworkGameplayTagNodeIndex.Swap(findidx, CommonIdx);
 				Found = true;
@@ -355,7 +321,7 @@ FName UGameplayTagsManager::GetTagNameFromNetIndex(FGameplayTagNetIndex Index)
 		ensureMsgf(Index == InvalidTagNetIndex, TEXT("Received invalid tag net index %d! Tag index is out of sync on client!"), Index);
 		return NAME_None;
 	}
-	return NetworkGameplayTagNodeIndex[Index]->GetCompleteTag();
+	return NetworkGameplayTagNodeIndex[Index]->GetCompleteTagName();
 }
 
 FGameplayTagNetIndex UGameplayTagsManager::GetNetIndexFromTag(const FGameplayTag &InTag)
@@ -372,20 +338,25 @@ FGameplayTagNetIndex UGameplayTagsManager::GetNetIndexFromTag(const FGameplayTag
 
 bool UGameplayTagsManager::ShouldImportTagsFromINI()
 {
+	UGameplayTagsSettings* MutableDefault = GetMutableDefault<UGameplayTagsSettings>();
+	FString DefaultEnginePath = FString::Printf(TEXT("%sDefaultEngine.ini"), *FPaths::SourceConfigDir());
+
+	// Deprecated path
 	bool ImportFromINI = false;
-	GConfig->GetBool(TEXT("GameplayTags"), TEXT("ImportTagsFromConfig"), ImportFromINI, GEngineIni);
-	return ImportFromINI;
-}
+	if (GConfig->GetBool(TEXT("GameplayTags"), TEXT("ImportTagsFromConfig"), ImportFromINI, DefaultEnginePath))
+	{
+		if (ImportFromINI)
+		{
+			MutableDefault->ImportTagsFromConfig = ImportFromINI;
+			UE_LOG(LogGameplayTags, Log, TEXT("ImportTagsFromConfig is in a deprecated location, open and save GameplayTag settings to fix"));
+		}
+		return ImportFromINI;
+	}
 
-/** TEMP - Returns true if we should warn on invalid (missing) tags */
-bool UGameplayTagsManager::ShouldWarnOnInvalidTags()
-{
-	bool ImportFromINI = true;
-	GConfig->GetBool(TEXT("GameplayTags"), TEXT("WarnOnInvalidTags"), ImportFromINI, GEngineIni);
-	return ImportFromINI;
+	return MutableDefault->ImportTagsFromConfig;
 }
 
-void UGameplayTagsManager::RedirectTagsForContainer(FGameplayTagContainer& Container, TSet<FName>& DeprecatedTagNamesNotFoundInTagMap, UProperty* SerializingProperty)
+void UGameplayTagsManager::RedirectTagsForContainer(FGameplayTagContainer& Container, UProperty* SerializingProperty)
 {
 	TSet<FName> NamesToRemove;
 	TSet<const FGameplayTag*> TagsToAdd;
@@ -401,7 +372,6 @@ void UGameplayTagsManager::RedirectTagsForContainer(FGameplayTagContainer& Conta
 			if (NewTag->IsValid())
 			{
 				TagsToAdd.Add(NewTag);
-				DeprecatedTagNamesNotFoundInTagMap.Remove(NewTag->GetTagName());
 			}
 		}
 #if WITH_EDITOR
@@ -417,16 +387,6 @@ void UGameplayTagsManager::RedirectTagsForContainer(FGameplayTagContainer& Conta
 #endif
 	}
 
-	// Add additional tags to the TagsToAdd set from the deprecated list if they weren't already added above
-	for (FName AdditionalDeprecatedTag : DeprecatedTagNamesNotFoundInTagMap)
-	{
-		const FGameplayTag* NewTag = TagRedirects.Find(AdditionalDeprecatedTag);
-		if (NewTag && NewTag->IsValid())
-		{
-			TagsToAdd.Add(NewTag);
-		}
-	}
-
 	// Remove all tags from the NamesToRemove set
 	for (FName RemoveName : NamesToRemove)
 	{
@@ -437,7 +397,7 @@ void UGameplayTagsManager::RedirectTagsForContainer(FGameplayTagContainer& Conta
 		}
 		else
 		{
-			FGameplayTagRedirectHelper::RemoveTagByExplicitName(Container, RemoveName);
+			Container.RemoveTagByExplicitName(RemoveName);
 		}
 	}
 
@@ -473,6 +433,33 @@ void UGameplayTagsManager::RedirectSingleGameplayTag(FGameplayTag& Tag, UPropert
 #endif
 }
 
+void UGameplayTagsManager::InitializeManager()
+{
+	check(!SingletonManager);
+
+	SingletonManager = NewObject<UGameplayTagsManager>(GetTransientPackage(), NAME_None);
+	SingletonManager->AddToRoot();
+
+	UGameplayTagsSettings* MutableDefault = GetMutableDefault<UGameplayTagsSettings>();
+	FString DefaultEnginePath = FString::Printf(TEXT("%sDefaultEngine.ini"), *FPaths::SourceConfigDir());
+
+	TArray<FString> GameplayTagTables;
+	GConfig->GetArray(TEXT("GameplayTags"), TEXT("+GameplayTagTableList"), GameplayTagTables, DefaultEnginePath);
+
+	// Report deprecation
+	if (GameplayTagTables.Num() > 0)
+	{
+		UE_LOG(LogGameplayTags, Log, TEXT("GameplayTagTableList is in a deprecated location, open and save GameplayTag settings to fix"));
+		for (const FString& DataTable : GameplayTagTables)
+		{
+			MutableDefault->GameplayTagTableList.AddUnique(DataTable);
+		}
+	}
+
+	SingletonManager->LoadGameplayTagTables();
+	SingletonManager->ConstructGameplayTagTree();
+}
+
 void UGameplayTagsManager::PopulateTreeFromDataTable(class UDataTable* InTable)
 {
 	checkf(GameplayRootTag.IsValid(), TEXT("ConstructGameplayTagTree() must be called before PopulateTreeFromDataTable()"));
@@ -481,51 +468,43 @@ void UGameplayTagsManager::PopulateTreeFromDataTable(class UDataTable* InTable)
 	TArray<FGameplayTagTableRow*> TagTableRows;
 	InTable->GetAllRows<FGameplayTagTableRow>(ContextString, TagTableRows);
 
+	FName SourceName = InTable->GetOutermost()->GetFName();
+
+	FGameplayTagSource* FoundSource = FindOrAddTagSource(SourceName, EGameplayTagSourceType::DataTable);
+
 	for (const FGameplayTagTableRow* TagRow : TagTableRows)
 	{
 		if (TagRow)
 		{
-			AddTagTableRow(*TagRow);
+			AddTagTableRow(*TagRow, SourceName);
 		}
 	}
-	GameplayRootTag->GetChildTagNodes().Sort(FCompareFGameplayTagNodeByTag());
 }
 
-void UGameplayTagsManager::AddTagTableRow(const FGameplayTagTableRow& TagRow)
+void UGameplayTagsManager::AddTagTableRow(const FGameplayTagTableRow& TagRow, FName SourceName)
 {
-	TArray< TSharedPtr<FGameplayTagNode> >& GameplayRootTags = GameplayRootTag->GetChildTagNodes();
+	TSharedPtr<FGameplayTagNode> CurNode = GameplayRootTag;
 
 	// Split the tag text on the "." delimiter to establish tag depth and then insert each tag into the
 	// gameplay tag tree
 	TArray<FString> SubTags;
-	TagRow.Tag.ParseIntoArray(SubTags, TEXT("."), true);
+	TagRow.Tag.ToString().ParseIntoArray(SubTags, TEXT("."), true);
 
-	if (SubTags.Num() > 0)
+	for (int32 SubTagIdx = 0; SubTagIdx < SubTags.Num(); ++SubTagIdx)
 	{
-		bool FromDictionary = SubTags.Num() == 1;
-		int32 InsertionIdx = InsertTagIntoNodeArray(*SubTags[0], NULL, GameplayRootTags, SubTags.Num() == 1 ? TagRow.CategoryText : FText(), FromDictionary);
-		TSharedPtr<FGameplayTagNode> CurNode = GameplayRootTags[InsertionIdx];
+		TArray< TSharedPtr<FGameplayTagNode> >& ChildTags = CurNode.Get()->GetChildTagNodes();
 
-		for (int32 SubTagIdx = 1; SubTagIdx < SubTags.Num(); ++SubTagIdx)
-		{
-			TArray< TSharedPtr<FGameplayTagNode> >& ChildTags = CurNode.Get()->GetChildTagNodes();
-			FText Description;
-			if (SubTagIdx == SubTags.Num() - 1)
-			{
-				Description = TagRow.CategoryText;
-			}
+		bool bFromDictionary = (SubTagIdx == (SubTags.Num() - 1));
+		int32 InsertionIdx = InsertTagIntoNodeArray(*SubTags[SubTagIdx], CurNode, ChildTags, bFromDictionary ? SourceName : NAME_None, TagRow.DevComment);
 
-			FromDictionary = (SubTagIdx == (SubTags.Num()-1));
-			InsertionIdx = InsertTagIntoNodeArray(*SubTags[SubTagIdx], CurNode, ChildTags, Description, FromDictionary);
-
-			CurNode = ChildTags[InsertionIdx];
-		}
+		CurNode = ChildTags[InsertionIdx];
 	}
 }
 
 UGameplayTagsManager::~UGameplayTagsManager()
 {
 	DestroyGameplayTagTree();
+	SingletonManager = nullptr;
 }
 
 void UGameplayTagsManager::DestroyGameplayTagTree()
@@ -537,14 +516,14 @@ void UGameplayTagsManager::DestroyGameplayTagTree()
 	}
 }
 
-int32 UGameplayTagsManager::InsertTagIntoNodeArray(FName Tag, TWeakPtr<FGameplayTagNode> ParentNode, TArray< TSharedPtr<FGameplayTagNode> >& NodeArray, FText CategoryDescription, bool InWasAddedDirectlyFromDictionary)
+int32 UGameplayTagsManager::InsertTagIntoNodeArray(FName Tag, TSharedPtr<FGameplayTagNode> ParentNode, TArray< TSharedPtr<FGameplayTagNode> >& NodeArray, FName SourceName, const FString& DevComment)
 {
 	int32 InsertionIdx = INDEX_NONE;
 
 	// See if the tag is already in the array
 	for (int32 CurIdx = 0; CurIdx < NodeArray.Num(); ++CurIdx)
 	{
-		if (NodeArray[CurIdx].IsValid() && NodeArray[CurIdx].Get()->GetSimpleTag() == Tag)
+		if (NodeArray[CurIdx].IsValid() && NodeArray[CurIdx].Get()->GetSimpleTagName() == Tag)
 		{
 			InsertionIdx = CurIdx;
 			break;
@@ -554,62 +533,36 @@ int32 UGameplayTagsManager::InsertTagIntoNodeArray(FName Tag, TWeakPtr<FGameplay
 	// Insert the tag at the end of the array if not already found
 	if (InsertionIdx == INDEX_NONE)
 	{
-		TSharedPtr<FGameplayTagNode> TagNode = MakeShareable(new FGameplayTagNode(Tag, ParentNode, CategoryDescription, InWasAddedDirectlyFromDictionary));
+		// Don't add the root node as parent
+		TSharedPtr<FGameplayTagNode> TagNode = MakeShareable(new FGameplayTagNode(Tag, ParentNode != GameplayRootTag ? ParentNode : nullptr));
 		InsertionIdx = NodeArray.Add(TagNode);
 
-		FGameplayTag GameplayTag = FGameplayTag(TagNode->GetCompleteTag());
+		FGameplayTag GameplayTag = TagNode->GetCompleteTag();
 
 		{
 #if WITH_EDITOR
-			// This critical section is to handle and editor-only issue where tag requests come from another thread when async loading from a background thread in FGameplayTagContainer::Serialize.
+			// This critical section is to handle an editor-only issue where tag requests come from another thread when async loading from a background thread in FGameplayTagContainer::Serialize.
 			// This function is not generically threadsafe.
 			FScopeLock Lock(&GameplayTagMapCritical);
 #endif
-			GameplayTagMap.Add(TagNode->GetCompleteTag(), GameplayTag);
-		}
-
-		GameplayTagNodeMap.Add(GameplayTag, TagNode);
-
-		TSet<FName> ParentSet;
-		GetAllParentNodeNames(ParentSet, TagNode);
-		for (auto Name : ParentSet)
-		{
-			TagNode->Parents.AddTagFast(FGameplayTag(Name));
+			GameplayTagNodeMap.Add(GameplayTag, TagNode);
 		}
 	}
-	else
-	{
-		// Fill in category description for nodes that were added by virtue of being a parent node
-		if (NodeArray[InsertionIdx]->CategoryDescription.IsEmpty() && !CategoryDescription.IsEmpty())
-		{
-			NodeArray[InsertionIdx]->CategoryDescription = CategoryDescription;
-		}
 
-		// This node could have been created as an implicit tag, and later we found the explicit declaration in the dictionary
-		NodeArray[InsertionIdx]->WasAddedDirectlyFromDictionary |= InWasAddedDirectlyFromDictionary;
+#if WITH_EDITOR
+	// Set/update editor only data
+	if (NodeArray[InsertionIdx]->SourceName.IsNone() && !SourceName.IsNone())
+	{
+		NodeArray[InsertionIdx]->SourceName = SourceName;
 	}
 
-	return InsertionIdx;
-}
-
-int32 UGameplayTagsManager::GetBestTagCategoryDescription(FString Tag, FText& OutDescription) const
-{
-	// get all the nodes that make up this tag
-	TArray< TSharedPtr<FGameplayTagNode> > TagItems;
-	GetAllNodesForTag(Tag, TagItems);
-	
-	// find the deepest tag in the list (last is deepest based on how this array was constructed)
-	TSharedPtr<FGameplayTagNode> BestDescriptionNode;
-	for(int32 TagIdx = TagItems.Num() - 1; TagIdx >= 0; TagIdx--)
+	if (NodeArray[InsertionIdx]->DevComment.IsEmpty() && !DevComment.IsEmpty())
 	{
-		OutDescription = TagItems[TagIdx]->GetCategoryDescription();
-		if(!OutDescription.IsEmpty())
-		{
-			return TagIdx;
-		}
+		NodeArray[InsertionIdx]->DevComment = DevComment;
 	}
+#endif
 
-	return -1;
+	return InsertionIdx;
 }
 
 #if !(UE_BUILD_SHIPPING || UE_BUILD_TEST)
@@ -742,7 +695,7 @@ static void RecursiveRootTagSearch(const FString& InFilterString, const TArray<T
 
 	for (int32 iTag = 0; iTag < GameplayRootTags.Num(); ++iTag)
 	{
-		FString RootTagName = GameplayRootTags[iTag].Get()->GetSimpleTag().ToString();
+		FString RootTagName = GameplayRootTags[iTag].Get()->GetSimpleTagName().ToString();
 
 		if (RootTagName.Equals(CurrentFilter) == true)
 		{
@@ -781,186 +734,112 @@ void UGameplayTagsManager::GetFilteredGameplayRootTags( const FString& InFilterS
 	}
 }
 
-void UGameplayTagsManager::OnObjectReimported(UFactory* ImportFactory, UObject* InObject)
+void UGameplayTagsManager::GetAllTagsFromSource(FName TagSource, TArray< TSharedPtr<FGameplayTagNode> >& OutTagArray)
 {
-	// Re-construct the gameplay tag tree if the base table is re-imported
-	if (GIsEditor && !IsRunningCommandlet() && InObject && GameplayTagTables.Contains(Cast<UDataTable>(InObject)))
+	for (const TPair<FGameplayTag, TSharedPtr<FGameplayTagNode>>& NodePair : GameplayTagNodeMap)
 	{
-		DestroyGameplayTagTree();
-		ConstructGameplayTagTree();
+		if (NodePair.Value->SourceName == TagSource)
+		{
+			OutTagArray.Add(NodePair.Value);
+		}
 	}
 }
 
-TSharedPtr<FGameplayTagNode> UGameplayTagsManager::FindTagNode(FName TagName) const
-{
-	return FindTagNode(GameplayRootTag, TagName);
-}
-
-void GameplayTagsUpdateSourceControl(FString RelativeConfigFilePath)
+bool UGameplayTagsManager::IsDictionaryTag(FName TagName) const
 {
-	FString ConfigPath = FPaths::ConvertRelativePathToFull(RelativeConfigFilePath);
-
-	if (ISourceControlModule::Get().IsEnabled())
+	const TSharedPtr<FGameplayTagNode> Node = FindTagNode(TagName);
 	{
-		FText ErrorMessage;
-
-		if (!SourceControlHelpers::CheckoutOrMarkForAdd(ConfigPath, FText::FromString(ConfigPath), NULL, ErrorMessage))
+		if (Node.IsValid() && Node->SourceName != NAME_None)
 		{
-			FNotificationInfo Info(ErrorMessage);
-			Info.ExpireDuration = 3.0f;
-			FSlateNotificationManager::Get().AddNotification(Info);
+			return true;
 		}
 	}
-	else
-	{
-		if (!FPlatformFileManager::Get().GetPlatformFile().SetReadOnly(*ConfigPath, false))
-		{
-			FText NotificationErrorText = FText::Format(LOCTEXT("FailedToMakeWritable", "Could not make {0} writable."), FText::FromString(ConfigPath));
-
-			FNotificationInfo Info(NotificationErrorText);
-			Info.ExpireDuration = 3.0f;
 
-			FSlateNotificationManager::Get().AddNotification(Info);
-		}
-	}
+	return false;
 }
 
-bool UGameplayTagsManager::IsDictionaryTag(FName TagName) const
+bool UGameplayTagsManager::GetTagEditorData(FName TagName, FString& OutComment, FName &OutTagSource) const
 {
-	FGameplayTag TagTemp;
-	TagTemp.TagName = TagName;
-	if (const TSharedPtr<FGameplayTagNode>* Node = GameplayTagNodeMap.Find(TagTemp))
+	const TSharedPtr<FGameplayTagNode> Node = FindTagNode(TagName);
 	{
-		if (Node->IsValid())
+		if (Node.IsValid())
 		{
-			if (Node->Get()->WasAddedDirectlyFromDictionary)
-			{
-				return true;
-			}
+			OutComment = Node->DevComment;
+			OutTagSource = Node->SourceName;
+			return true;
 		}
 	}
-
 	return false;
 }
 
-void UGameplayTagsManager::AddNewGameplayTagToINI(FString NewTag)
+void UGameplayTagsManager::EditorRefreshGameplayTagTree()
 {
-	if(NewTag.IsEmpty())
-	{
-		return;
-	}
+	DestroyGameplayTagTree();
+	LoadGameplayTagTables();
+	ConstructGameplayTagTree();
+}
 
-	if (ShouldImportTagsFromINI() == false)
+FGameplayTagContainer UGameplayTagsManager::RequestGameplayTagChildrenInDictionary(const FGameplayTag& GameplayTag) const
+{
+	FGameplayTagContainer TagContainer;
+	// Note this purposefully does not include the passed in GameplayTag in the container.
+	const TSharedPtr<FGameplayTagNode>* GameplayTagNode = GameplayTagNodeMap.Find(GameplayTag);
+	if (GameplayTagNode)
 	{
-		return;
+		AddChildrenTags(TagContainer, *GameplayTagNode, true, true);
 	}
-	
-	UGameplayTagsSettings*			Settings = GetMutableDefault<UGameplayTagsSettings>();
-	UGameplayTagsDeveloperSettings* DevSettings = GetMutableDefault<UGameplayTagsDeveloperSettings>();
+	return TagContainer;
+}
 
-	// Already in the list as an explicit tag, ignore. Note we want to add if it is in implicit tag. (E.g, someone added A.B.C then someone tries to add A.B)
+#endif // WITH_EDITOR
+
+const FGameplayTagSource* UGameplayTagsManager::FindTagSource(FName TagSourceName) const
+{
+	for (const FGameplayTagSource& TagSource : TagSources)
 	{
-		
-		if (IGameplayTagsModule::Get().GetGameplayTagsManager().IsDictionaryTag(FName(*NewTag)))
+		if (TagSource.SourceName == TagSourceName)
 		{
-			return;
+			return &TagSource;
 		}
 	}
+	return nullptr;
+}
 
-	if (Settings)
+void UGameplayTagsManager::FindTagSourcesWithType(EGameplayTagSourceType TagSourceType, TArray<const FGameplayTagSource*>& OutArray) const
+{
+	for (const FGameplayTagSource& TagSource : TagSources)
 	{
-		FString RelativeConfigFilePath = Settings->GetDefaultConfigFilename();
-		if (DevSettings && DevSettings->DeveloperConfigName.IsEmpty() == false)
+		if (TagSource.SourceType == TagSourceType)
 		{
-			// Write this tag to the Developer's ini file
-			FDeveloperTags& DeveloperTags = IGameplayTagsModule::Get().GetGameplayTagsManager().DeveloperTags;
-
-			// Get config filename
-			RelativeConfigFilePath.ReplaceInline(TEXT("DefaultGameplayTags"), *FString::Printf(TEXT("Tags/%s"), *DevSettings->DeveloperConfigName));
-
-			// See if there is already an entry for this one
-			int32 idx=0;
-			for(; idx < DeveloperTags.Items.Num(); ++idx)
-			{
-				if (DeveloperTags.Items[idx].IniName == RelativeConfigFilePath)
-				{
-					break;
-				}
-			}
-
-			// Create new entry if necessary
-			if (DeveloperTags.Items.IsValidIndex(idx) == false)
-			{
-				idx = DeveloperTags.Items.AddDefaulted();
-			}
-
-			// Setup the entry and add the tag
-			DeveloperTags.Items[idx].IniName = RelativeConfigFilePath;
-			DeveloperTags.Items[idx].Tags.Add(NewTag);
-
-			// Write back out to config
-			GConfig->SetArray(TEXT("UserTags"), TEXT("GameplayTags"), DeveloperTags.Items[idx].Tags, RelativeConfigFilePath);
-
-			// Flush once to possible write the file to disk for the first time
-			GConfig->Flush(false, RelativeConfigFilePath);
-
-			// Source Control
-			GameplayTagsUpdateSourceControl(RelativeConfigFilePath);
-
-			// Flush again in the case where the file already existed
-			GConfig->Flush(false, RelativeConfigFilePath);
-			
-		}
-		else
-		{
-			// Write this tag to the global ini file
-			Settings->GameplayTags.Add(NewTag);
-			Settings->SortTags();
-
-			// Backup the full list before dumping to config
-			TArray<FString> Backup = Settings->GameplayTags;
-
-			// Remove DeveloperTags from the master list
-			FDeveloperTags& DeveloperTags = IGameplayTagsModule::Get().GetGameplayTagsManager().DeveloperTags;
-			for(int32 idx=0; idx < DeveloperTags.Items.Num(); ++idx)
-			{
-				for (FString& Str : DeveloperTags.Items[idx].Tags)
-				{
-					Settings->GameplayTags.Remove(Str);
-				}
-			}			
-						
-			// Source Control
-			GameplayTagsUpdateSourceControl(RelativeConfigFilePath);
-
-			Settings->UpdateDefaultConfigFile();
-			
-			GConfig->Flush(false, RelativeConfigFilePath);
-
-			// Restore
-			Settings->GameplayTags = Backup;
+			OutArray.Add(&TagSource);
 		}
+	}
+}
 
-		// ------------------------------------------------
-
-		
+FGameplayTagSource* UGameplayTagsManager::FindOrAddTagSource(FName TagSourceName, EGameplayTagSourceType SourceType)
+{
+	const FGameplayTagSource* FoundSource = FindTagSource(TagSourceName);
+	if (FoundSource)
+	{
+		return const_cast<FGameplayTagSource*>(FoundSource);
+	}
 
-		// ------------------------------------------------
+	// Need to make a new one
 
-		IGameplayTagsModule::Get().GetGameplayTagsManager().DestroyGameplayTagTree();
+	FGameplayTagSource* NewSource = new(TagSources) FGameplayTagSource(TagSourceName, SourceType);
 
-		{
-#if STATS
-			FString PerfMessage = FString::Printf(TEXT("ConstructGameplayTagTree GameplayTag tables after adding new tag"));
-			SCOPE_LOG_TIME_IN_SECONDS(*PerfMessage, nullptr)
-#endif
-			IGameplayTagsModule::Get().GetGameplayTagsManager().ConstructGameplayTagTree();
-		}
+	if (SourceType == EGameplayTagSourceType::DefaultTagList)
+	{
+		NewSource->SourceTagList = GetMutableDefault<UGameplayTagsSettings>();
+	}
+	else if (SourceType == EGameplayTagSourceType::TagList)
+	{
+		NewSource->SourceTagList = NewObject<UGameplayTagsList>(this, TagSourceName, RF_Transient);
+		NewSource->SourceTagList->ConfigFileName = FString::Printf(TEXT("%sTags/%s"), *FPaths::SourceConfigDir(), *TagSourceName.ToString());
 	}
-	
-}
 
-#endif // WITH_EDITOR
+	return NewSource;
+}
 
 DECLARE_CYCLE_STAT(TEXT("UGameplayTagsManager::RequestGameplayTag"), STAT_UGameplayTagsManager_RequestGameplayTag, STATGROUP_GameplayTags);
 
@@ -973,135 +852,87 @@ FGameplayTag UGameplayTagsManager::RequestGameplayTag(FName TagName, bool ErrorI
 	FScopeLock Lock(&GameplayTagMapCritical);
 #endif
 
-	const FGameplayTag* Tag = GameplayTagMap.Find(TagName);
-	if (!Tag)
-	{ 
-		if (ErrorIfNotFound)
+	FGameplayTag PossibleTag(TagName);
+
+	if (GameplayTagNodeMap.Contains(PossibleTag))
+	{
+		return PossibleTag;
+	}
+	else if (ErrorIfNotFound)
+	{
+		static TSet<FName> MissingTagName;
+		if (!MissingTagName.Contains(TagName))
 		{
-			static TSet<FName> MissingTagName;
-			if (!MissingTagName.Contains(TagName))
-			{
-				ensureAlwaysMsgf(false, TEXT("Requested Tag %s was not found. Check tag data table."), *TagName.ToString());
-				MissingTagName.Add(TagName);
-			}
+			ensureAlwaysMsgf(false, TEXT("Requested Tag %s was not found. Check tag data table."), *TagName.ToString());
+			MissingTagName.Add(TagName);
 		}
-		return FGameplayTag();
 	}
-	return *Tag;
+	return FGameplayTag();
 }
 
 FGameplayTagContainer UGameplayTagsManager::RequestGameplayTagParents(const FGameplayTag& GameplayTag) const
 {
-	FGameplayTagContainer TagContainer;
-	TagContainer.AddTag(GameplayTag);
-	AddParentTags(TagContainer, GameplayTag);
-	return TagContainer;
+	const FGameplayTagContainer* ParentTags = GetSingleTagContainer(GameplayTag);
+
+	if (ParentTags)
+	{
+		return ParentTags->GetGameplayTagParents();
+	}
+	return FGameplayTagContainer();
 }
 
-FGameplayTagContainer UGameplayTagsManager::RequestGameplayTagChildren(const FGameplayTag& GameplayTag, bool OnlyIncludeDictionaryTags) const
+FGameplayTagContainer UGameplayTagsManager::RequestGameplayTagChildren(const FGameplayTag& GameplayTag) const
 {
 	FGameplayTagContainer TagContainer;
 	// Note this purposefully does not include the passed in GameplayTag in the container.
-	AddChildrenTags(TagContainer, GameplayTag, true, OnlyIncludeDictionaryTags);
-	return TagContainer;
-}
-
-FGameplayTag UGameplayTagsManager::RequestGameplayTagDirectParent(const FGameplayTag& GameplayTag) const
-{
 	const TSharedPtr<FGameplayTagNode>* GameplayTagNode = GameplayTagNodeMap.Find(GameplayTag);
 	if (GameplayTagNode)
 	{
-		TSharedPtr<FGameplayTagNode> Parent = (*GameplayTagNode)->GetParentTagNode().Pin();
-		if (Parent.IsValid())
-		{
-			const FGameplayTag* Tag = GameplayTagNodeMap.FindKey(Parent);
-			if (Tag)
-			{
-				return *Tag;
-			}
-		}
+		AddChildrenTags(TagContainer, *GameplayTagNode, true, false);
 	}
-	return FGameplayTag();
-}
-
-bool UGameplayTagsManager::AddLeafTagToContainer(FGameplayTagContainer& TagContainer, const FGameplayTag& Tag)
-{
-	// Check tag is not already in container
-	if (TagContainer.HasTag(Tag, EGameplayTagMatchType::Explicit, EGameplayTagMatchType::Explicit))
-	{
-		return true;
-	}
-
-	// Check tag is not just a parent of an existing tag in the container
-	for (auto It = TagContainer.CreateConstIterator(); It; ++It)
-	{
-		FGameplayTagContainer ParentTags = RequestGameplayTagParents(*It);
-		if (ParentTags.HasTag(Tag, EGameplayTagMatchType::Explicit, EGameplayTagMatchType::Explicit))
-		{
-			return false;
-		}
-	}
-
-	// Remove any tags in the container that are a parent to this tag
-	FGameplayTagContainer ParentTags = RequestGameplayTagParents(Tag);
-	for (auto It = ParentTags.CreateConstIterator(); It; ++It)
-	{
-		if (TagContainer.HasTag(*It, EGameplayTagMatchType::Explicit, EGameplayTagMatchType::Explicit))
-		{
-			TagContainer.RemoveTag(*It);
-		}
-	}
-
-	// Add the tag
-	TagContainer.AddTag(Tag);
-	return true;
+	return TagContainer;
 }
 
-DECLARE_CYCLE_STAT(TEXT("AddParentTags"), STAT_AddParentTags, STATGROUP_GameplayTags);
-
-
-void UGameplayTagsManager::AddParentTags(FGameplayTagContainer& TagContainer, const FGameplayTag& GameplayTag) const
+FGameplayTag UGameplayTagsManager::RequestGameplayTagDirectParent(const FGameplayTag& GameplayTag) const
 {
-	SCOPE_CYCLE_COUNTER(STAT_AddParentTags);
-
 	const TSharedPtr<FGameplayTagNode>* GameplayTagNode = GameplayTagNodeMap.Find(GameplayTag);
 	if (GameplayTagNode)
 	{
-		TSharedPtr<FGameplayTagNode> Parent = (*GameplayTagNode)->GetParentTagNode().Pin();
+		TSharedPtr<FGameplayTagNode> Parent = (*GameplayTagNode)->GetParentTagNode();
 		if (Parent.IsValid())
 		{
-			const FGameplayTag* Tag = GameplayTagNodeMap.FindKey(Parent);
-			if (Tag)
-			{
-				TagContainer.AddTag(*Tag);
-				AddParentTags(TagContainer, *Tag);
-			}
+			return Parent->GetCompleteTag();
 		}
 	}
+	return FGameplayTag();
 }
 
-void UGameplayTagsManager::AddChildrenTags(FGameplayTagContainer& TagContainer, const FGameplayTag& GameplayTag, bool RecurseAll, bool OnlyIncludeDictionaryTags) const
+void UGameplayTagsManager::AddChildrenTags(FGameplayTagContainer& TagContainer, const TSharedPtr<FGameplayTagNode> GameplayTagNode, bool RecurseAll, bool OnlyIncludeDictionaryTags) const
 {
-	const TSharedPtr<FGameplayTagNode>* GameplayTagNode = GameplayTagNodeMap.Find(GameplayTag);
-	if (GameplayTagNode)
+	if (GameplayTagNode.IsValid())
 	{
-		TArray< TSharedPtr<FGameplayTagNode> >& ChildrenNodes = (*GameplayTagNode)->GetChildTagNodes();
+		TArray< TSharedPtr<FGameplayTagNode> >& ChildrenNodes = GameplayTagNode->GetChildTagNodes();
 		for (TSharedPtr<FGameplayTagNode> ChildNode : ChildrenNodes)
 		{
 			if (ChildNode.IsValid())
 			{
-				const FGameplayTag* Tag = GameplayTagNodeMap.FindKey(ChildNode);
-				if (Tag)
+				bool bShouldInclude = true;
+
+				if (bShouldInclude)
 				{
-					if (OnlyIncludeDictionaryTags == false || ChildNode->WasAddedDirectlyFromDictionary)
-					{
-						TagContainer.AddTag(*Tag);
-					}
+					TagContainer.AddTag(ChildNode->GetCompleteTag());
+				}
 
-					if (RecurseAll)
-					{
-						AddChildrenTags(TagContainer, *Tag, true, OnlyIncludeDictionaryTags);
-					}
+#if WITH_EDITORONLY_DATA
+				if (OnlyIncludeDictionaryTags && ChildNode->SourceName == NAME_None)
+				{
+					// Only have info to do this in editor builds
+					bShouldInclude = false;
+				}
+#endif			
+				if (RecurseAll)
+				{
+					AddChildrenTags(TagContainer, ChildNode, true, OnlyIncludeDictionaryTags);
 				}
 			}
 
@@ -1116,126 +947,12 @@ void UGameplayTagsManager::SplitGameplayTagFName(const FGameplayTag& Tag, TArray
 		TSharedPtr<FGameplayTagNode> CurNode = *GameplayTagNodePtr;
 		while (CurNode.IsValid())
 		{
-			OutNames.Insert(CurNode->GetSimpleTag(), 0);
-			CurNode = CurNode->GetParentTagNode().Pin();
-		}
-	}
-}
-
-DECLARE_CYCLE_STAT(TEXT("UGameplayTagsManager::ComplexGameplayTagsMatch ParentsParents"), STAT_UGameplayTagsManager_GameplayTagsMatchParentsParents, STATGROUP_GameplayTags);
-
-bool UGameplayTagsManager::ComplexGameplayTagsMatch(const FGameplayTag& GameplayTagOne, TEnumAsByte<EGameplayTagMatchType::Type> MatchTypeOne, const FGameplayTag& GameplayTagTwo, TEnumAsByte<EGameplayTagMatchType::Type> MatchTypeTwo) const
-{
-	check(MatchTypeOne != EGameplayTagMatchType::Explicit || MatchTypeTwo != EGameplayTagMatchType::Explicit);
-
-	if (MatchTypeOne == EGameplayTagMatchType::IncludeParentTags && MatchTypeTwo == EGameplayTagMatchType::IncludeParentTags)
-	{
-		SCOPE_CYCLE_COUNTER(STAT_UGameplayTagsManager_GameplayTagsMatchParentsParents);
-		const TSharedPtr<FGameplayTagNode>* TagNodeOne = GameplayTagNodeMap.Find(GameplayTagOne);
-		const TSharedPtr<FGameplayTagNode>* TagNodeTwo = GameplayTagNodeMap.Find(GameplayTagTwo);
-		if (TagNodeOne && TagNodeTwo)
-		{
-			return (*TagNodeOne)->Parents.DoesTagContainerMatch((*TagNodeTwo)->Parents, EGameplayTagMatchType::Explicit, EGameplayTagMatchType::Explicit, EGameplayContainerMatchType::Any);
-		}
-	}
-	else
-	{
-		checkSlow(MatchTypeOne != EGameplayTagMatchType::Explicit || MatchTypeTwo != EGameplayTagMatchType::Explicit);
-
-		if (MatchTypeOne == EGameplayTagMatchType::IncludeParentTags)
-		{
-			const TSharedPtr<FGameplayTagNode>* TagNodeOne = GameplayTagNodeMap.Find(GameplayTagOne);
-			if (TagNodeOne)
-			{
-				return (*TagNodeOne)->Parents.HasTag(GameplayTagTwo, EGameplayTagMatchType::Explicit, EGameplayTagMatchType::Explicit);
-			}
-		}
-		else
-		{
-			const TSharedPtr<FGameplayTagNode>* TagNodeTwo = GameplayTagNodeMap.Find(GameplayTagTwo);
-			if (TagNodeTwo)
-			{
-				return (*TagNodeTwo)->Parents.HasTag(GameplayTagOne, EGameplayTagMatchType::Explicit, EGameplayTagMatchType::Explicit);
-			}
+			OutNames.Insert(CurNode->GetSimpleTagName(), 0);
+			CurNode = CurNode->GetParentTagNode();
 		}
 	}
-	return false;
 }
 
-#if CHECK_TAG_OPTIMIZATIONS
-
-bool UGameplayTagsManager::GameplayTagsMatchOriginal(const FGameplayTag& GameplayTagOne, TEnumAsByte<EGameplayTagMatchType::Type> MatchTypeOne, const FGameplayTag& GameplayTagTwo, TEnumAsByte<EGameplayTagMatchType::Type> MatchTypeTwo) const
-{
-	SCOPE_CYCLE_COUNTER(STAT_UGameplayTagsManager_GameplayTagsMatch);
-	if (MatchTypeOne == EGameplayTagMatchType::Explicit && MatchTypeTwo == EGameplayTagMatchType::Explicit)
-	{
-		const TSharedPtr<FGameplayTagNode>* TagNode1 = GameplayTagNodeMap.Find(GameplayTagOne);
-		const TSharedPtr<FGameplayTagNode>* TagNode2 = GameplayTagNodeMap.Find(GameplayTagTwo);
-		if (TagNode1 && TagNode2)
-		{
-			return (*TagNode1)->GetCompleteTag() == (*TagNode2)->GetCompleteTag();
-		}
-		return false;
-	}
-	static TSet<FName> Tags1;
-	static TSet<FName> Tags2;
-
-	check(!Tags1.Num() && !Tags2.Num()); // must be game thread, cannot call this recursively
-
-
-	if (MatchTypeOne == EGameplayTagMatchType::IncludeParentTags)
-	{
-		const TSharedPtr<FGameplayTagNode>* TagNode = GameplayTagNodeMap.Find(GameplayTagOne);
-		if (TagNode)
-		{
-			GetAllParentNodeNames(Tags1, *TagNode);
-		}
-	}
-	if (MatchTypeTwo == EGameplayTagMatchType::IncludeParentTags)
-	{
-		const TSharedPtr<FGameplayTagNode>* TagNode = GameplayTagNodeMap.Find(GameplayTagTwo);
-		if (TagNode)
-		{
-			GetAllParentNodeNames(Tags2, *TagNode);
-		}
-	}
-	bool bResult = false;
-
-	if (MatchTypeOne == EGameplayTagMatchType::Explicit)
-	{
-		const TSharedPtr<FGameplayTagNode>* TagNode = GameplayTagNodeMap.Find(GameplayTagOne);
-		if (TagNode)
-		{
-			bResult = Tags2.Contains((*TagNode)->GetCompleteTag());
-		}
-	}
-	else if (MatchTypeTwo == EGameplayTagMatchType::Explicit)
-	{
-		const TSharedPtr<FGameplayTagNode>* TagNode = GameplayTagNodeMap.Find(GameplayTagTwo);
-		if (TagNode)
-		{
-			bResult = Tags1.Contains((*TagNode)->GetCompleteTag());
-		}
-	}
-	else
-	{
-		for (auto& Tag1 : Tags1)
-		{
-			if (Tags2.Contains(Tag1))
-			{
-				bResult = true;
-				break;
-			}
-		}
-	}
-
-	Tags1.Reset();
-	Tags2.Reset();
-	return bResult;
-}
-
-#endif
-
 int32 UGameplayTagsManager::GameplayTagsMatchDepth(const FGameplayTag& GameplayTagOne, const FGameplayTag& GameplayTagTwo) const
 {
 	TSet<FName> Tags1;
@@ -1262,8 +979,8 @@ void UGameplayTagsManager::GetAllParentNodeNames(TSet<FName>& NamesList, const T
 {
 	SCOPE_CYCLE_COUNTER(STAT_UGameplayTagsManager_GetAllParentNodeNames);
 
-	NamesList.Add(GameplayTag->GetCompleteTag());
-	const TSharedPtr<FGameplayTagNode> Parent = GameplayTag->GetParentTagNode().Pin();
+	NamesList.Add(GameplayTag->GetCompleteTagName());
+	const TSharedPtr<FGameplayTagNode> Parent = GameplayTag->GetParentTagNode();
 	if (Parent.IsValid())
 	{
 		GetAllParentNodeNames(NamesList, Parent);
@@ -1276,28 +993,18 @@ bool UGameplayTagsManager::ValidateTagCreation(FName TagName) const
 {
 	SCOPE_CYCLE_COUNTER(STAT_UGameplayTagsManager_ValidateTagCreation);
 
-	return FindTagNode(GameplayRootTag, TagName).IsValid();
+	return FindTagNode(TagName).IsValid();
 }
 
-TSharedPtr<FGameplayTagNode> UGameplayTagsManager::FindTagNode(TSharedPtr<FGameplayTagNode> Node, FName TagName) const
+TSharedPtr<FGameplayTagNode> UGameplayTagsManager::FindTagNode(FName TagName) const
 {
-	TSharedPtr<FGameplayTagNode> FoundTag = NULL;
-	const TArray< TSharedPtr<FGameplayTagNode> >& CurrentChildrenTags = Node->GetChildTagNodes();
-	for (int32 TagIdx = 0; TagIdx < CurrentChildrenTags.Num(); TagIdx++)
+	FGameplayTag PossibleTag(TagName);
+
+	if (const TSharedPtr<FGameplayTagNode>* Node = GameplayTagNodeMap.Find(PossibleTag))
 	{
-		const TSharedPtr<FGameplayTagNode>& TagNode = CurrentChildrenTags[TagIdx];
-		if (TagNode->GetCompleteTag() == TagName || TagNode->GetSimpleTag() == TagName)
-		{
-			FoundTag = CurrentChildrenTags[TagIdx];
-			break;
-		}
-		FoundTag = FindTagNode(TagNode, TagName);
-		if (FoundTag.IsValid())
-		{
-			break;
-		}
+		return *Node;
 	}
-	return FoundTag;
+	return nullptr;
 }
 
 FGameplayTagTableRow::FGameplayTagTableRow(FGameplayTagTableRow const& Other)
@@ -1314,7 +1021,6 @@ FGameplayTagTableRow& FGameplayTagTableRow::operator=(FGameplayTagTableRow const
 	}
 
 	Tag = Other.Tag;
-	CategoryText = Other.CategoryText;
 	DevComment = Other.DevComment;
 
 	return *this;
@@ -1330,43 +1036,44 @@ bool FGameplayTagTableRow::operator!=(FGameplayTagTableRow const& Other) const
 	return (Tag != Other.Tag);
 }
 
-FGameplayTagNode::FGameplayTagNode(FName InTag, TWeakPtr<FGameplayTagNode> InParentNode, FText InCategoryDescription, bool InWasAddedDirectlyFromDictionary)
+bool FGameplayTagTableRow::operator<(FGameplayTagTableRow const& Other) const
+{
+	return (Tag < Other.Tag);
+}
+
+FGameplayTagNode::FGameplayTagNode(FName InTag, TSharedPtr<FGameplayTagNode> InParentNode)
 	: Tag(InTag)
-	, CompleteTag(NAME_None)
-	, CategoryDescription(InCategoryDescription)
 	, ParentNode(InParentNode)
 	, NetIndex(INVALID_TAGNETINDEX)
 {
-	WasAddedDirectlyFromDictionary = InWasAddedDirectlyFromDictionary;
-
-	TArray<FName> Tags;
+	TArray<FGameplayTag> ParentCompleteTags;
 
-	Tags.Add(InTag);
+	TSharedPtr<FGameplayTagNode> CurNode = InParentNode;
 
-	TWeakPtr<FGameplayTagNode> CurNode = InParentNode;
-	while (CurNode.IsValid())
+	// Stop iterating at root node
+	while (CurNode.IsValid() && CurNode->GetSimpleTagName() != NAME_None)
 	{
-		Tags.Add(CurNode.Pin()->GetSimpleTag());
-		CurNode = CurNode.Pin()->GetParentTagNode();
+		ParentCompleteTags.Add(CurNode->GetCompleteTag());
+		CurNode = CurNode->GetParentTagNode();
 	}
 
-	FString CompleteTagString;
-	for (int32 TagIdx = Tags.Num() - 1; TagIdx >= 0; --TagIdx)
+	FString CompleteTagString = InTag.ToString();
+
+	if (ParentCompleteTags.Num() > 0)
 	{
-		CompleteTagString += Tags[TagIdx].ToString();
-		if (TagIdx > 0)
-		{
-			CompleteTagString += TEXT(".");
-		}
+		// If we have a parent, add parent., which will includes all earlier parents
+		CompleteTagString = FString::Printf(TEXT("%s.%s"), *ParentCompleteTags[0].ToString(), *InTag.ToString());
 	}
-
-	CompleteTag = FName(*CompleteTagString);
+	
+	// Manually construct the tag container as we want to bypass the safety checks
+	CompleteTagWithParents.GameplayTags.Add(FGameplayTag(FName(*CompleteTagString)));
+	CompleteTagWithParents.ParentTags = ParentCompleteTags;
 }
 
 void FGameplayTagNode::ResetNode()
 {
 	Tag = NAME_None;
-	CompleteTag = NAME_None;
+	CompleteTagWithParents.Reset();
 	NetIndex = INVALID_TAGNETINDEX;
 
 	for (int32 ChildIdx = 0; ChildIdx < ChildTags.Num(); ++ChildIdx)
diff --git a/Engine/Source/Runtime/GameplayTags/Private/GameplayTagsModule.cpp b/Engine/Source/Runtime/GameplayTags/Private/GameplayTagsModule.cpp
--- a/Engine/Source/Runtime/GameplayTags/Private/GameplayTagsModule.cpp
+++ b/Engine/Source/Runtime/GameplayTags/Private/GameplayTagsModule.cpp
@@ -2,6 +2,9 @@
 
 #include "GameplayTagsModulePrivatePCH.h"
 
+FSimpleMulticastDelegate IGameplayTagsModule::OnGameplayTagTreeChanged;
+FSimpleMulticastDelegate IGameplayTagsModule::OnTagSettingsChanged;
+
 class FGameplayTagsModule : public IGameplayTagsModule
 {
 	// Begin IModuleInterface
@@ -15,15 +18,8 @@ DEFINE_LOG_CATEGORY(LogGameplayTags);
 
 void FGameplayTagsModule::StartupModule()
 {
-	// This code will execute after your module is loaded into memory (but after global variables are initialized, of course.)
-	GGameplayTagsManager = NewObject<UGameplayTagsManager>(GetTransientPackage(), NAME_None);
-	GGameplayTagsManager->AddToRoot();
-
-	TArray<FString> GameplayTagTables;
-	GConfig->GetArray(TEXT("GameplayTags"), TEXT("GameplayTagTableList"), GameplayTagTables, GEngineIni);
-
-	GGameplayTagsManager->LoadGameplayTagTables(GameplayTagTables);
-	GGameplayTagsManager->ConstructGameplayTagTree();
+	// This will force initialization
+	UGameplayTagsManager::Get();
 }
 
 #if !(UE_BUILD_SHIPPING || UE_BUILD_TEST)
@@ -37,9 +33,9 @@ void FGameplayTagsModule::ShutdownModule()
 #if !(UE_BUILD_SHIPPING || UE_BUILD_TEST)
 	if (GameplayTagPrintReportOnShutdown)
 	{
-		UGameplayTagsManager::PrintReplicationFrequencyReport();
+		UGameplayTagsManager::Get().PrintReplicationFrequencyReport();
 	}
 #endif
 
-	GGameplayTagsManager = NULL;
+	UGameplayTagsManager::SingletonManager = nullptr;
 }
diff --git a/Engine/Source/Runtime/GameplayTags/Private/GameplayTagsModulePrivatePCH.h b/Engine/Source/Runtime/GameplayTags/Private/GameplayTagsModulePrivatePCH.h
--- a/Engine/Source/Runtime/GameplayTags/Private/GameplayTagsModulePrivatePCH.h
+++ b/Engine/Source/Runtime/GameplayTags/Private/GameplayTagsModulePrivatePCH.h
@@ -2,12 +2,6 @@
 
 #pragma once
 
-
-
-#if WITH_EDITOR
-#include "UnrealEd.h"
-#endif
-
 // You should place include statements to your module's private header files here.  You only need to
 // add includes for headers that are used in most of your module's source files though.
 
diff --git a/Engine/Source/Runtime/GameplayTags/Private/GameplayTagsSettings.cpp b/Engine/Source/Runtime/GameplayTags/Private/GameplayTagsSettings.cpp
--- a/Engine/Source/Runtime/GameplayTags/Private/GameplayTagsSettings.cpp
+++ b/Engine/Source/Runtime/GameplayTags/Private/GameplayTagsSettings.cpp
@@ -1,11 +1,26 @@
 // Copyright 1998-2016 Epic Games, Inc. All Rights Reserved.
 
 #include "GameplayTagsModulePrivatePCH.h"
-#include "GameplayTagsModule.h"
+
+UGameplayTagsList::UGameplayTagsList(const FObjectInitializer& ObjectInitializer)
+	: Super(ObjectInitializer)
+{
+	// No config filename, needs to be set at creation time
+}
+
+void UGameplayTagsList::SortTags()
+{
+	GameplayTagList.Sort();
+}
 
 UGameplayTagsSettings::UGameplayTagsSettings(const FObjectInitializer& ObjectInitializer)
 : Super(ObjectInitializer)
 {
+	ConfigFileName = GetDefaultConfigFilename();
+	ImportTagsFromConfig = false;
+	WarnOnInvalidTags = true;
+	FastReplication = false;
+	NumBitsForContainerSize = 6;
 	NetIndexFirstBitSegment = 16;
 }
 
@@ -15,20 +30,11 @@ void UGameplayTagsSettings::PostEditChangeProperty(FPropertyChangedEvent& Proper
 	Super::PostEditChangeProperty(PropertyChangedEvent);
 	if (PropertyChangedEvent.Property)
 	{
-		if (PropertyChangedEvent.Property->GetFName() == GET_MEMBER_NAME_CHECKED(UGameplayTagsSettings, GameplayTags))
-		{
-			IGameplayTagsModule::Get().GetGameplayTagsManager().DestroyGameplayTagTree();
-			IGameplayTagsModule::Get().GetGameplayTagsManager().ConstructGameplayTagTree();
-		}
+		IGameplayTagsModule::OnTagSettingsChanged.Broadcast();
 	}
 }
 #endif
 
-void UGameplayTagsSettings::SortTags()
-{
-	GameplayTags.Sort();
-}
-
 // ---------------------------------
 
 UGameplayTagsDeveloperSettings::UGameplayTagsDeveloperSettings(const FObjectInitializer& ObjectInitializer)
