Copying //UE4/Dev-Framework to //UE4/Dev-Main (Source: //UE4/Dev-Framework @ 3182037)

diff --git a/Engine/Source/Runtime/GameplayTags/Private/Tests/GameplayTagTests.cpp b/Engine/Source/Runtime/GameplayTags/Private/Tests/GameplayTagTests.cpp
--- a/Engine/Source/Runtime/GameplayTags/Private/Tests/GameplayTagTests.cpp
+++ b/Engine/Source/Runtime/GameplayTags/Private/Tests/GameplayTagTests.cpp
@@ -4,269 +4,361 @@
 
 #if WITH_DEV_AUTOMATION_TESTS
 
-#if WITH_EDITOR
-
-#define ADD_TAG_CSV(Name, Count) CSV.Append(FString::Printf(TEXT("\r\n%d,"#Name), Count++));
-
-static UDataTable* CreateGameplayDataTable()
+class FGameplayTagTestBase : public FAutomationTestBase
 {
-	int32 Count = 0;
-	FString CSV(TEXT(",Tag,CategoryText,"));
-	ADD_TAG_CSV(Effect.Damage, Count);
-	ADD_TAG_CSV(Effect.Damage.Basic, Count);
-	ADD_TAG_CSV(Effect.Damage.Type1, Count);
-	ADD_TAG_CSV(Effect.Damage.Type2, Count);
-	ADD_TAG_CSV(Effect.Damage.Reduce, Count);
-	ADD_TAG_CSV(Effect.Damage.Buffable, Count);
-	ADD_TAG_CSV(Effect.Damage.Buff, Count);
-	ADD_TAG_CSV(Effect.Damage.Physical, Count);
-	ADD_TAG_CSV(Effect.Damage.Fire, Count);
-	ADD_TAG_CSV(Effect.Damage.Buffed.FireBuff, Count);
-	ADD_TAG_CSV(Effect.Damage.Mitigated.Armor, Count);
-	ADD_TAG_CSV(Effect.Lifesteal, Count);
-	ADD_TAG_CSV(Effect.Shield, Count);
-	ADD_TAG_CSV(Effect.Buff, Count);
-	ADD_TAG_CSV(Effect.Immune, Count);
-	ADD_TAG_CSV(Effect.FireDamage, Count);
-	ADD_TAG_CSV(Effect.Shield.Absorb, Count);
-	ADD_TAG_CSV(Effect.Protect.Damage, Count);
-	ADD_TAG_CSV(Stackable, Count);
-	ADD_TAG_CSV(Stack.DiminishingReturns, Count);
-	ADD_TAG_CSV(GameplayCue.Burning, Count);
-	ADD_TAG_CSV(Expensive.Status.Tag.Type.1, Count);
-	ADD_TAG_CSV(Expensive.Status.Tag.Type.2, Count);
-	ADD_TAG_CSV(Expensive.Status.Tag.Type.3, Count);
-	ADD_TAG_CSV(Expensive.Status.Tag.Type.4, Count);
-	ADD_TAG_CSV(Expensive.Status.Tag.Type.5, Count);
-	ADD_TAG_CSV(Expensive.Status.Tag.Type.6, Count);
-	ADD_TAG_CSV(Expensive.Status.Tag.Type.7, Count);
-	ADD_TAG_CSV(Expensive.Status.Tag.Type.8, Count);
-	ADD_TAG_CSV(Expensive.Status.Tag.Type.9, Count);
-	ADD_TAG_CSV(Expensive.Status.Tag.Type.10, Count);
-	ADD_TAG_CSV(Expensive.Status.Tag.Type.11, Count);
-	ADD_TAG_CSV(Expensive.Status.Tag.Type.12, Count);
-	ADD_TAG_CSV(Expensive.Status.Tag.Type.13, Count);
-	ADD_TAG_CSV(Expensive.Status.Tag.Type.14, Count);
-	ADD_TAG_CSV(Expensive.Status.Tag.Type.15, Count);
-	ADD_TAG_CSV(Expensive.Status.Tag.Type.16, Count);
-	ADD_TAG_CSV(Expensive.Status.Tag.Type.17, Count);
-	ADD_TAG_CSV(Expensive.Status.Tag.Type.18, Count);
-	ADD_TAG_CSV(Expensive.Status.Tag.Type.19, Count);
-	ADD_TAG_CSV(Expensive.Status.Tag.Type.20, Count);
-	ADD_TAG_CSV(Expensive.Status.Tag.Type.21, Count);
-	ADD_TAG_CSV(Expensive.Status.Tag.Type.22, Count);
-	ADD_TAG_CSV(Expensive.Status.Tag.Type.23, Count);
-	ADD_TAG_CSV(Expensive.Status.Tag.Type.24, Count);
-	ADD_TAG_CSV(Expensive.Status.Tag.Type.25, Count);
-	ADD_TAG_CSV(Expensive.Status.Tag.Type.26, Count);
-	ADD_TAG_CSV(Expensive.Status.Tag.Type.27, Count);
-	ADD_TAG_CSV(Expensive.Status.Tag.Type.28, Count);
-	ADD_TAG_CSV(Expensive.Status.Tag.Type.29, Count);
-	ADD_TAG_CSV(Expensive.Status.Tag.Type.30, Count);
-	ADD_TAG_CSV(Expensive.Status.Tag.Type.31, Count);
-	ADD_TAG_CSV(Expensive.Status.Tag.Type.32, Count);
-	ADD_TAG_CSV(Expensive.Status.Tag.Type.33, Count);
-	ADD_TAG_CSV(Expensive.Status.Tag.Type.34, Count);
-	ADD_TAG_CSV(Expensive.Status.Tag.Type.35, Count);
-	ADD_TAG_CSV(Expensive.Status.Tag.Type.36, Count);
-	ADD_TAG_CSV(Expensive.Status.Tag.Type.37, Count);
-	ADD_TAG_CSV(Expensive.Status.Tag.Type.38, Count);
-	ADD_TAG_CSV(Expensive.Status.Tag.Type.39, Count);
-	ADD_TAG_CSV(Expensive.Status.Tag.Type.40, Count);
-
-	auto DataTable = NewObject<UDataTable>(GetTransientPackage(), FName(TEXT("TempDataTable")));
-	DataTable->RowStruct = FGameplayTagTableRow::StaticStruct();
-	DataTable->CreateTableFromCSVString(CSV);
-
-	FGameplayTagTableRow * Row = (FGameplayTagTableRow*)DataTable->RowMap["0"];
-	if (Row)
+public:
+	FGameplayTagTestBase(const FString& InName, const bool bInComplexTask)
+		: FAutomationTestBase(InName, bInComplexTask)
 	{
-		check(Row->Tag == TEXT("Effect.Damage"));
 	}
-	return DataTable;
-}
 
-static FGameplayTag GetTagForString(const FString& String)
-{
-	return IGameplayTagsModule::Get().GetGameplayTagsManager().RequestGameplayTag(FName(*String));
-}
+	UDataTable* CreateGameplayDataTable()
+	{	
+		TArray<FString> TestTags;
+
+		TestTags.Add(TEXT("Effect.Damage"));
+		TestTags.Add(TEXT("Effect.Damage.Basic"));
+		TestTags.Add(TEXT("Effect.Damage.Type1"));
+		TestTags.Add(TEXT("Effect.Damage.Type2"));
+		TestTags.Add(TEXT("Effect.Damage.Reduce"));
+		TestTags.Add(TEXT("Effect.Damage.Buffable"));
+		TestTags.Add(TEXT("Effect.Damage.Buff"));
+		TestTags.Add(TEXT("Effect.Damage.Physical"));
+		TestTags.Add(TEXT("Effect.Damage.Fire"));
+		TestTags.Add(TEXT("Effect.Damage.Buffed.FireBuff"));
+		TestTags.Add(TEXT("Effect.Damage.Mitigated.Armor"));
+		TestTags.Add(TEXT("Effect.Lifesteal"));
+		TestTags.Add(TEXT("Effect.Shield"));
+		TestTags.Add(TEXT("Effect.Buff"));
+		TestTags.Add(TEXT("Effect.Immune"));
+		TestTags.Add(TEXT("Effect.FireDamage"));
+		TestTags.Add(TEXT("Effect.Shield.Absorb"));
+		TestTags.Add(TEXT("Effect.Protect.Damage"));
+		TestTags.Add(TEXT("Stackable"));
+		TestTags.Add(TEXT("Stack.DiminishingReturns"));
+		TestTags.Add(TEXT("GameplayCue.Burning"));
+		TestTags.Add(TEXT("Expensive.Status.Tag.Type.1"));
+		TestTags.Add(TEXT("Expensive.Status.Tag.Type.2"));
+		TestTags.Add(TEXT("Expensive.Status.Tag.Type.3"));
+		TestTags.Add(TEXT("Expensive.Status.Tag.Type.4"));
+		TestTags.Add(TEXT("Expensive.Status.Tag.Type.5"));
+		TestTags.Add(TEXT("Expensive.Status.Tag.Type.6"));
+		TestTags.Add(TEXT("Expensive.Status.Tag.Type.7"));
+		TestTags.Add(TEXT("Expensive.Status.Tag.Type.8"));
+		TestTags.Add(TEXT("Expensive.Status.Tag.Type.9"));
+		TestTags.Add(TEXT("Expensive.Status.Tag.Type.10"));
+		TestTags.Add(TEXT("Expensive.Status.Tag.Type.11"));
+		TestTags.Add(TEXT("Expensive.Status.Tag.Type.12"));
+		TestTags.Add(TEXT("Expensive.Status.Tag.Type.13"));
+		TestTags.Add(TEXT("Expensive.Status.Tag.Type.14"));
+		TestTags.Add(TEXT("Expensive.Status.Tag.Type.15"));
+		TestTags.Add(TEXT("Expensive.Status.Tag.Type.16"));
+		TestTags.Add(TEXT("Expensive.Status.Tag.Type.17"));
+		TestTags.Add(TEXT("Expensive.Status.Tag.Type.18"));
+		TestTags.Add(TEXT("Expensive.Status.Tag.Type.19"));
+		TestTags.Add(TEXT("Expensive.Status.Tag.Type.20"));
+		TestTags.Add(TEXT("Expensive.Status.Tag.Type.21"));
+		TestTags.Add(TEXT("Expensive.Status.Tag.Type.22"));
+		TestTags.Add(TEXT("Expensive.Status.Tag.Type.23"));
+		TestTags.Add(TEXT("Expensive.Status.Tag.Type.24"));
+		TestTags.Add(TEXT("Expensive.Status.Tag.Type.25"));
+		TestTags.Add(TEXT("Expensive.Status.Tag.Type.26"));
+		TestTags.Add(TEXT("Expensive.Status.Tag.Type.27"));
+		TestTags.Add(TEXT("Expensive.Status.Tag.Type.28"));
+		TestTags.Add(TEXT("Expensive.Status.Tag.Type.29"));
+		TestTags.Add(TEXT("Expensive.Status.Tag.Type.30"));
+		TestTags.Add(TEXT("Expensive.Status.Tag.Type.31"));
+		TestTags.Add(TEXT("Expensive.Status.Tag.Type.32"));
+		TestTags.Add(TEXT("Expensive.Status.Tag.Type.33"));
+		TestTags.Add(TEXT("Expensive.Status.Tag.Type.34"));
+		TestTags.Add(TEXT("Expensive.Status.Tag.Type.35"));
+		TestTags.Add(TEXT("Expensive.Status.Tag.Type.36"));
+		TestTags.Add(TEXT("Expensive.Status.Tag.Type.37"));
+		TestTags.Add(TEXT("Expensive.Status.Tag.Type.38"));
+		TestTags.Add(TEXT("Expensive.Status.Tag.Type.39"));
+		TestTags.Add(TEXT("Expensive.Status.Tag.Type.40"));
+
+		auto DataTable = NewObject<UDataTable>(GetTransientPackage(), FName(TEXT("TempDataTable")));
+		DataTable->RowStruct = FGameplayTagTableRow::StaticStruct();
 
-static bool GameplayTagTest_SimpleTest()
-{
-	FName TagName = FName(TEXT("Stack.DiminishingReturns"));
-	FGameplayTag Tag = IGameplayTagsModule::Get().GetGameplayTagsManager().RequestGameplayTag(TagName);
-	return Tag.GetTagName() == TagName;
-}
+#if WITH_EDITOR
+		FString CSV(TEXT(",Tag,CategoryText,"));
+		for (int32 Count = 0; Count < TestTags.Num(); Count++)
+		{
+			CSV.Append(FString::Printf(TEXT("\r\n%d,%s"), Count, *TestTags[Count]));
+		}
 
-static bool GameplayTagTest_TagComparisonTest()
-{
-	FGameplayTag EffectDamageTag = GetTagForString(TEXT("Effect.Damage"));
-	FGameplayTag EffectDamage1Tag = GetTagForString(TEXT("Effect.Damage.Type1"));
-	FGameplayTag EffectDamage2Tag = GetTagForString(TEXT("Effect.Damage.Type2"));
-	FGameplayTag CueTag = GetTagForString(TEXT("GameplayCue.Burning"));
+		DataTable->CreateTableFromCSVString(CSV);
+#else
+		// Construct the data table manually
+		for (int32 Count = 0; Count < TestTags.Num(); Count++)
+		{
+			uint8* RowData = (uint8*)FMemory::Malloc(DataTable->RowStruct->PropertiesSize);
+			DataTable->RowStruct->InitializeStruct(RowData);
+
+			FGameplayTagTableRow* TagRow = (FGameplayTagTableRow*)RowData;
+			TagRow->Tag = FName(*TestTags[Count]);
+
+			DataTable->RowMap.Add(FName(*FString::FromInt(Count)), RowData);
+		}
+#endif
 	
-	bool bResult = true;
+		FGameplayTagTableRow * Row = (FGameplayTagTableRow*)DataTable->RowMap["0"];
+		if (Row)
+		{
+			check(Row->Tag == TEXT("Effect.Damage"));
+		}
+		return DataTable;
+	}
 
-	bResult &= (EffectDamage1Tag == EffectDamage1Tag);
-	bResult &= (EffectDamage1Tag != EffectDamage2Tag);
-	bResult &= (EffectDamage1Tag != EffectDamageTag);
+	FGameplayTag GetTagForString(const FString& String)
+	{
+		return UGameplayTagsManager::Get().RequestGameplayTag(FName(*String));
+	}
 
-	bResult &= (EffectDamage1Tag.Matches(EGameplayTagMatchType::IncludeParentTags, EffectDamageTag, EGameplayTagMatchType::Explicit));
-	bResult &= (EffectDamage1Tag.Matches(EGameplayTagMatchType::IncludeParentTags, EffectDamageTag, EGameplayTagMatchType::IncludeParentTags));
-	bResult &= (!EffectDamage1Tag.Matches(EGameplayTagMatchType::Explicit, EffectDamageTag, EGameplayTagMatchType::IncludeParentTags));
-	bResult &= (!EffectDamage1Tag.Matches(EGameplayTagMatchType::Explicit, EffectDamageTag, EGameplayTagMatchType::Explicit));
+	void GameplayTagTest_SimpleTest()
+	{
+		FName TagName = FName(TEXT("Stack.DiminishingReturns"));
+		FGameplayTag Tag = UGameplayTagsManager::Get().RequestGameplayTag(TagName);
+		TestTrueExpr(Tag.GetTagName() == TagName);
+	}
 
-	bResult &= (!EffectDamage1Tag.Matches(EGameplayTagMatchType::IncludeParentTags, CueTag, EGameplayTagMatchType::IncludeParentTags));
+	void GameplayTagTest_TagComparisonTest()
+	{
+		FGameplayTag EffectDamageTag = GetTagForString(TEXT("Effect.Damage"));
+		FGameplayTag EffectDamage1Tag = GetTagForString(TEXT("Effect.Damage.Type1"));
+		FGameplayTag EffectDamage2Tag = GetTagForString(TEXT("Effect.Damage.Type2"));
+		FGameplayTag CueTag = GetTagForString(TEXT("GameplayCue.Burning"));
+		FGameplayTag EmptyTag;
+
+		TestTrueExpr(EffectDamage1Tag == EffectDamage1Tag);
+		TestTrueExpr(EffectDamage1Tag != EffectDamage2Tag);
+		TestTrueExpr(EffectDamage1Tag != EffectDamageTag);
+
+		TestTrueExpr(EffectDamage1Tag.MatchesTag(EffectDamageTag));
+		TestTrueExpr(!EffectDamage1Tag.MatchesTagExact(EffectDamageTag));
+		TestTrueExpr(!EffectDamage1Tag.MatchesTag(EmptyTag));
+		TestTrueExpr(!EffectDamage1Tag.MatchesTagExact(EmptyTag));
+		TestTrueExpr(!EmptyTag.MatchesTag(EmptyTag));
+		TestTrueExpr(!EmptyTag.MatchesTagExact(EmptyTag));
+
+		TestTrueExpr(EffectDamage1Tag.RequestDirectParent() == EffectDamageTag);
+	}
 
-	bResult &= (EffectDamage1Tag.RequestDirectParent() == EffectDamageTag);
+	void GameplayTagTest_TagContainerTest()
+	{
+		FGameplayTag EffectDamageTag = GetTagForString(TEXT("Effect.Damage"));
+		FGameplayTag EffectDamage1Tag = GetTagForString(TEXT("Effect.Damage.Type1"));
+		FGameplayTag EffectDamage2Tag = GetTagForString(TEXT("Effect.Damage.Type2"));
+		FGameplayTag CueTag = GetTagForString(TEXT("GameplayCue.Burning"));
+		FGameplayTag EmptyTag;
 
-	return bResult;
-}
+		FGameplayTagContainer EmptyContainer;
 
-static bool GameplayTagTest_TagContainerTest()
-{
-	FGameplayTag EffectDamageTag = GetTagForString(TEXT("Effect.Damage"));
-	FGameplayTag EffectDamage1Tag = GetTagForString(TEXT("Effect.Damage.Type1"));
-	FGameplayTag EffectDamage2Tag = GetTagForString(TEXT("Effect.Damage.Type2"));
-	FGameplayTag CueTag = GetTagForString(TEXT("GameplayCue.Burning"));
+		FGameplayTagContainer TagContainer;
+		TagContainer.AddTag(EffectDamage1Tag);
+		TagContainer.AddTag(CueTag);
 
-	bool bResult = true;
+		FGameplayTagContainer ReverseTagContainer;
+		ReverseTagContainer.AddTag(CueTag);
+		ReverseTagContainer.AddTag(EffectDamage1Tag);
+	
+		FGameplayTagContainer TagContainer2;
+		TagContainer2.AddTag(EffectDamage2Tag);
+		TagContainer2.AddTag(CueTag);
 
-	FGameplayTagContainer TagContainer;
-	TagContainer.AddTag(EffectDamage1Tag);
-	TagContainer.AddTag(CueTag);
+		TestTrueExpr(TagContainer == TagContainer);
+		TestTrueExpr(TagContainer == ReverseTagContainer);
+		TestTrueExpr(TagContainer != TagContainer2);
 
-	FGameplayTagContainer TagContainer2;
-	TagContainer2.AddTag(EffectDamage2Tag);
-	TagContainer2.AddTag(CueTag);
+		FGameplayTagContainer TagContainerCopy = TagContainer;
 
-	FGameplayTagContainer EmptyContainer;
+		TestTrueExpr(TagContainerCopy == TagContainer);
+		TestTrueExpr(TagContainerCopy != TagContainer2);
 
-	bResult &= (TagContainer == TagContainer);
-	bResult &= (TagContainer != TagContainer2);
+		TagContainerCopy.Reset();
+		TagContainerCopy.AppendTags(TagContainer);
 
-	bResult &= (TagContainer.MatchesAny(TagContainer2, true));
-	bResult &= (!TagContainer.MatchesAll(TagContainer2, true));
+		TestTrueExpr(TagContainerCopy == TagContainer);
+		TestTrueExpr(TagContainerCopy != TagContainer2);
 
-	bResult &= (TagContainer.MatchesAll(EmptyContainer, true));
-	bResult &= (!TagContainer.MatchesAll(EmptyContainer, false));
-	bResult &= (!EmptyContainer.MatchesAll(TagContainer, true));
-	bResult &= (!EmptyContainer.MatchesAll(TagContainer, false));
+		TestTrueExpr(TagContainer.HasAny(TagContainer2));
+		TestTrueExpr(TagContainer.HasAnyExact(TagContainer2));
+		TestTrueExpr(!TagContainer.HasAll(TagContainer2));
+		TestTrueExpr(!TagContainer.HasAllExact(TagContainer2));
+		TestTrueExpr(TagContainer.HasAll(TagContainerCopy));
+		TestTrueExpr(TagContainer.HasAllExact(TagContainerCopy));
 
-	bResult &= (TagContainer.HasTagExplicit(EffectDamage1Tag));
-	bResult &= (!TagContainer.HasTagExplicit(EffectDamage2Tag));
+		TestTrueExpr(TagContainer.HasAll(EmptyContainer));
+		TestTrueExpr(TagContainer.HasAllExact(EmptyContainer));
+		TestTrueExpr(!TagContainer.HasAny(EmptyContainer));
+		TestTrueExpr(!TagContainer.HasAnyExact(EmptyContainer));
 
-	bResult &= (TagContainer.HasTag(EffectDamageTag, EGameplayTagMatchType::IncludeParentTags, EGameplayTagMatchType::Explicit));
-	bResult &= (TagContainer.HasTag(EffectDamageTag, EGameplayTagMatchType::IncludeParentTags, EGameplayTagMatchType::IncludeParentTags));
-	bResult &= (!TagContainer.HasTag(EffectDamageTag, EGameplayTagMatchType::Explicit, EGameplayTagMatchType::Explicit));
-	bResult &= (!TagContainer.HasTag(EffectDamageTag, EGameplayTagMatchType::Explicit, EGameplayTagMatchType::IncludeParentTags));
+		TestTrueExpr(EmptyContainer.HasAll(EmptyContainer));
+		TestTrueExpr(EmptyContainer.HasAllExact(EmptyContainer));
+		TestTrueExpr(!EmptyContainer.HasAny(EmptyContainer));
+		TestTrueExpr(!EmptyContainer.HasAnyExact(EmptyContainer));
 
-	bResult &= (!TagContainer.HasTag(EffectDamage2Tag, EGameplayTagMatchType::IncludeParentTags, EGameplayTagMatchType::Explicit));
-	bResult &= (TagContainer.HasTag(EffectDamage2Tag, EGameplayTagMatchType::IncludeParentTags, EGameplayTagMatchType::IncludeParentTags));
-	bResult &= (!TagContainer.HasTag(EffectDamage2Tag, EGameplayTagMatchType::Explicit, EGameplayTagMatchType::Explicit));
-	bResult &= (!TagContainer.HasTag(EffectDamage2Tag, EGameplayTagMatchType::Explicit, EGameplayTagMatchType::IncludeParentTags));
+		TestTrueExpr(!EmptyContainer.HasAll(TagContainer));
+		TestTrueExpr(!EmptyContainer.HasAllExact(TagContainer));
+		TestTrueExpr(!EmptyContainer.HasAny(TagContainer));
+		TestTrueExpr(!EmptyContainer.HasAnyExact(TagContainer));
 
-	FGameplayTagContainer FilteredTagContainer = TagContainer.Filter(TagContainer2, EGameplayTagMatchType::Explicit, EGameplayTagMatchType::Explicit);
-	bResult &= (FilteredTagContainer.HasTagExplicit(CueTag));
-	bResult &= (!FilteredTagContainer.HasTagExplicit(EffectDamage1Tag));
+		TestTrueExpr(TagContainer.HasTag(EffectDamageTag));
+		TestTrueExpr(!TagContainer.HasTagExact(EffectDamageTag));
+		TestTrueExpr(!TagContainer.HasTag(EmptyTag));
+		TestTrueExpr(!TagContainer.HasTagExact(EmptyTag));
 
-	FilteredTagContainer = TagContainer.Filter(TagContainer2, EGameplayTagMatchType::IncludeParentTags, EGameplayTagMatchType::IncludeParentTags);
-	bResult &= (FilteredTagContainer.HasTagExplicit(CueTag));
-	bResult &= (FilteredTagContainer.HasTagExplicit(EffectDamage1Tag));
+		TestTrueExpr(EffectDamage1Tag.MatchesAny(FGameplayTagContainer(EffectDamageTag)));
+		TestTrueExpr(!EffectDamage1Tag.MatchesAnyExact(FGameplayTagContainer(EffectDamageTag)));
 
-	return bResult;
-}
+		TestTrueExpr(EffectDamage1Tag.MatchesAny(TagContainer));
 
-static bool GameplayTagTest_PerfTest()
-{
-	FGameplayTag EffectDamageTag = GetTagForString(TEXT("Effect.Damage"));
-	FGameplayTag EffectDamage1Tag = GetTagForString(TEXT("Effect.Damage.Type1"));
-	FGameplayTag EffectDamage2Tag = GetTagForString(TEXT("Effect.Damage.Type2"));
-	FGameplayTag CueTag = GetTagForString(TEXT("GameplayCue.Burning"));
+		FGameplayTagContainer FilteredTagContainer = TagContainer.FilterExact(TagContainer2);
+		TestTrueExpr(FilteredTagContainer.HasTagExact(CueTag));
+		TestTrueExpr(!FilteredTagContainer.HasTagExact(EffectDamage1Tag));
 
-	FGameplayTagContainer TagContainer;
+		FilteredTagContainer = TagContainer.Filter(FGameplayTagContainer(EffectDamageTag));
+		TestTrueExpr(!FilteredTagContainer.HasTagExact(CueTag));
+		TestTrueExpr(FilteredTagContainer.HasTagExact(EffectDamage1Tag));
 
-	bool bResult = true;
+		FilteredTagContainer.Reset();
+		FilteredTagContainer.AppendMatchingTags(TagContainer, TagContainer2);
+	
+		TestTrueExpr(FilteredTagContainer.HasTagExact(CueTag));
+		TestTrueExpr(!FilteredTagContainer.HasTagExact(EffectDamage1Tag));
+	}
 
+	void GameplayTagTest_PerfTest()
 	{
-		FScopeLogTime LogTimePtr(*FString::Printf(TEXT("1000 container constructions")), nullptr, FScopeLogTime::ScopeLog_Milliseconds);
-		for (int32 i = 0; i < 1000; i++)
+		FGameplayTag EffectDamageTag = GetTagForString(TEXT("Effect.Damage"));
+		FGameplayTag EffectDamage1Tag = GetTagForString(TEXT("Effect.Damage.Type1"));
+		FGameplayTag EffectDamage2Tag = GetTagForString(TEXT("Effect.Damage.Type2"));
+		FGameplayTag CueTag = GetTagForString(TEXT("GameplayCue.Burning"));
+
+		FGameplayTagContainer TagContainer;
+
+		bool bResult = true;
+
 		{
-			TagContainer = FGameplayTagContainer();
-			TagContainer.AddTag(EffectDamage1Tag);
-			TagContainer.AddTag(EffectDamage2Tag);
-			TagContainer.AddTag(CueTag);
-			for (int32 j = 1; j <= 40; j++)
+			FScopeLogTime LogTimePtr(*FString::Printf(TEXT("10000 get tag")), nullptr, FScopeLogTime::ScopeLog_Milliseconds);
+			for (int32 i = 0; i < 10000; i++)
 			{
-				TagContainer.AddTag(GetTagForString(FString::Printf(TEXT("Expensive.Status.Tag.Type.%d"), j)));
+				UGameplayTagsManager::Get().RequestGameplayTag(FName(TEXT("Effect.Damage")));
+			}
+		}
+
+		{
+			FScopeLogTime LogTimePtr(*FString::Printf(TEXT("1000 container constructions")), nullptr, FScopeLogTime::ScopeLog_Milliseconds);
+			for (int32 i = 0; i < 1000; i++)
+			{
+				TagContainer = FGameplayTagContainer();
+				TagContainer.AddTag(EffectDamage1Tag);
+				TagContainer.AddTag(EffectDamage2Tag);
+				TagContainer.AddTag(CueTag);
+				for (int32 j = 1; j <= 40; j++)
+				{
+					TagContainer.AddTag(GetTagForString(FString::Printf(TEXT("Expensive.Status.Tag.Type.%d"), j)));
+				}
 			}
 		}
-	}
 
-	{
-		FScopeLogTime LogTimePtr(*FString::Printf(TEXT("1000 container copies")), nullptr, FScopeLogTime::ScopeLog_Milliseconds);
-		for (int32 i = 0; i < 1000; i++)
 		{
-			FGameplayTagContainer TagContainerNew;
+			FScopeLogTime LogTimePtr(*FString::Printf(TEXT("1000 container copies")), nullptr, FScopeLogTime::ScopeLog_Milliseconds);
+			for (int32 i = 0; i < 1000; i++)
+			{
+				FGameplayTagContainer TagContainerNew;
+
+				for (auto It = TagContainer.CreateConstIterator(); It; ++It)
+				{
+					TagContainerNew.AddTag(*It);
+				}
+			}
+		}
 
-			for (auto It = TagContainer.CreateConstIterator(); It; ++It)
+		{
+			FScopeLogTime LogTimePtr(*FString::Printf(TEXT("1000 container appends")), nullptr, FScopeLogTime::ScopeLog_Milliseconds);
+			for (int32 i = 0; i < 1000; i++)
 			{
-				TagContainerNew.AddTag(*It);
+				FGameplayTagContainer TagContainerNew;
+
+				TagContainerNew.AppendTags(TagContainer);
 			}
 		}
-	}
 	
-	FGameplayTagContainer TagContainer2;
-	TagContainer2.AddTag(EffectDamage1Tag);
-	TagContainer2.AddTag(EffectDamage2Tag);
-	TagContainer2.AddTag(CueTag);
+		FGameplayTagContainer TagContainer2;
+		TagContainer2.AddTag(EffectDamage1Tag);
+		TagContainer2.AddTag(EffectDamage2Tag);
+		TagContainer2.AddTag(CueTag);
 
-	{
-		FScopeLogTime LogTimePtr(*FString::Printf(TEXT("10000 explicit tag checks")), nullptr, FScopeLogTime::ScopeLog_Milliseconds);
-		for (int32 i = 0; i < 10000; i++)
 		{
-			bResult &= TagContainer.HasTag(EffectDamage1Tag, EGameplayTagMatchType::Explicit, EGameplayTagMatchType::Explicit);
+			FScopeLogTime LogTimePtr(*FString::Printf(TEXT("10000 MatchesAnyExact checks")), nullptr, FScopeLogTime::ScopeLog_Milliseconds);
+			for (int32 i = 0; i < 10000; i++)
+			{
+				bResult &= EffectDamage1Tag.MatchesAnyExact(TagContainer);
+			}
 		}
-	}
 
-	{
-		FScopeLogTime LogTimePtr(*FString::Printf(TEXT("10000 parent tag checks")), nullptr, FScopeLogTime::ScopeLog_Milliseconds);
-		for (int32 i = 0; i < 10000; i++)
 		{
-			bResult &= TagContainer.HasTag(EffectDamage1Tag, EGameplayTagMatchType::IncludeParentTags, EGameplayTagMatchType::IncludeParentTags);
+			FScopeLogTime LogTimePtr(*FString::Printf(TEXT("10000 MatchesAny checks")), nullptr, FScopeLogTime::ScopeLog_Milliseconds);
+			for (int32 i = 0; i < 10000; i++)
+			{
+				bResult &= EffectDamage1Tag.MatchesAny(TagContainer);
+			}
+		}
+
+		{
+			FScopeLogTime LogTimePtr(*FString::Printf(TEXT("10000 HasTagExact checks")), nullptr, FScopeLogTime::ScopeLog_Milliseconds);
+			for (int32 i = 0; i < 10000; i++)
+			{
+				bResult &= TagContainer.HasTagExact(EffectDamage1Tag);
+			}
+		}
+
+		{
+			FScopeLogTime LogTimePtr(*FString::Printf(TEXT("10000 HasTag checks")), nullptr, FScopeLogTime::ScopeLog_Milliseconds);
+			for (int32 i = 0; i < 10000; i++)
+			{
+				bResult &= TagContainer.HasTag(EffectDamage1Tag);
+			}
+		}
+
+		{
+			FScopeLogTime LogTimePtr(*FString::Printf(TEXT("10000 HasAll checks")), nullptr, FScopeLogTime::ScopeLog_Milliseconds);
+			for (int32 i = 0; i < 10000; i++)
+			{
+				bResult &= TagContainer.HasAll(TagContainer2);
+			}
 		}
-	}
 
-	{
-		FScopeLogTime LogTimePtr(*FString::Printf(TEXT("10000 container checks")), nullptr, FScopeLogTime::ScopeLog_Milliseconds);
-		for (int32 i = 0; i < 10000; i++)
 		{
-			bResult &= TagContainer.MatchesAll(TagContainer2, true);
+			FScopeLogTime LogTimePtr(*FString::Printf(TEXT("10000 HasAny checks")), nullptr, FScopeLogTime::ScopeLog_Milliseconds);
+			for (int32 i = 0; i < 10000; i++)
+			{
+				bResult &= TagContainer.HasAny(TagContainer2);
+			}
 		}
+
+		TestTrue(TEXT("Performance Tests succeeded"), bResult);
 	}
 
-	return bResult;
-}
+};
 
-IMPLEMENT_SIMPLE_AUTOMATION_TEST(FGameplayTagTest, "System.GameplayTags.GameplayTag", EAutomationTestFlags::EditorContext | EAutomationTestFlags::EngineFilter)
+IMPLEMENT_CUSTOM_SIMPLE_AUTOMATION_TEST(FGameplayTagTest, FGameplayTagTestBase, "System.GameplayTags.GameplayTag", EAutomationTestFlags::ApplicationContextMask | EAutomationTestFlags::EngineFilter)
 bool FGameplayTagTest::RunTest(const FString& Parameters)
 {
 	// Create Test Data 
 	UDataTable* DataTable = CreateGameplayDataTable();
 
-	IGameplayTagsModule::Get().GetGameplayTagsManager().PopulateTreeFromDataTable(DataTable);
+	UGameplayTagsManager::Get().PopulateTreeFromDataTable(DataTable);
 
 	// Run Tests
-	bool bSuccess = true;
-	bSuccess &= GameplayTagTest_SimpleTest();
-	bSuccess &= GameplayTagTest_TagComparisonTest();
-	bSuccess &= GameplayTagTest_TagContainerTest();
-	bSuccess &= GameplayTagTest_PerfTest();
-	// Add more tests here... 
-
-	return bSuccess;
-}
+	GameplayTagTest_SimpleTest();
+	GameplayTagTest_TagComparisonTest();
+	GameplayTagTest_TagContainerTest();
+	GameplayTagTest_PerfTest();
 
-#endif //WITH_EDITOR 
+	return !HasAnyErrors();
+}
 
 #endif //WITH_DEV_AUTOMATION_TESTS
\ No newline at end of file
diff --git a/Engine/Source/Runtime/GameplayTags/Public/GameplayTagsModule.h b/Engine/Source/Runtime/GameplayTags/Public/GameplayTagsModule.h
--- a/Engine/Source/Runtime/GameplayTags/Public/GameplayTagsModule.h
+++ b/Engine/Source/Runtime/GameplayTags/Public/GameplayTagsModule.h
@@ -6,11 +6,8 @@
 #include "EngineGlobals.h"
 #include "GameplayTagsManager.h"
 
-
-class UGameplayTagsManager;
-
 /**
- * The public interface to this module
+ * The public interface to this module, generally you should access the manager directly instead
  */
 class IGameplayTagsModule : public IModuleInterface
 {
@@ -40,29 +37,22 @@ public:
 		return FModuleManager::Get().IsModuleLoaded(GameplayTagModuleName);
 	}
 
-	/**
-	 * Helper function to request a gameplay tag by name
-	 * 
-	 * @param InTagName	Tag name to request
-	 * 
-	 * @return Gameplay tag associated with the specified name. Will be marked invalid if tag not found
-	 */
+	/** Delegate for when assets are added to the tree */
+	static GAMEPLAYTAGS_API FSimpleMulticastDelegate OnGameplayTagTreeChanged;
+
+	/** Delegate that gets called after the settings have changed in the editor */
+	static GAMEPLAYTAGS_API FSimpleMulticastDelegate OnTagSettingsChanged;
+
+	DEPRECATED(4.15, "Call FGameplayTag::RequestGameplayTag or RequestGameplayTag on the manager instead")
 	FORCEINLINE_DEBUGGABLE static FGameplayTag RequestGameplayTag(FName InTagName, bool ErrorIfNotFound=true)
 	{
-		IGameplayTagsModule& GameplayTagsModule = IGameplayTagsModule::Get();
-		return GameplayTagsModule.GetGameplayTagsManager().RequestGameplayTag(InTagName, ErrorIfNotFound);
+		return UGameplayTagsManager::Get().RequestGameplayTag(InTagName, ErrorIfNotFound);
 	}
 
-
-	// Gets the UGameplayTagsManager manager
+	DEPRECATED(4.15, "Call UGameplayTagsManager::Get instead")
 	FORCEINLINE_DEBUGGABLE static UGameplayTagsManager& GetGameplayTagsManager()
 	{
-		if (GGameplayTagsManager == nullptr)
-		{
-			IGameplayTagsModule::Get();
-		}
-		
-		return *GGameplayTagsManager;
+		return UGameplayTagsManager::Get();
 	}
 
 };
diff --git a/Engine/Source/Runtime/Launch/Private/LaunchEngineLoop.cpp b/Engine/Source/Runtime/Launch/Private/LaunchEngineLoop.cpp
--- a/Engine/Source/Runtime/Launch/Private/LaunchEngineLoop.cpp
+++ b/Engine/Source/Runtime/Launch/Private/LaunchEngineLoop.cpp
@@ -2201,9 +2201,6 @@ bool FEngineLoop::LoadStartupCoreModules()
 		FModuleManager::Get().LoadModule(TEXT("Blutility"));
 	}
 
-	// Needed for extra Blueprint nodes that can be used in standalone.
-	FModuleManager::Get().LoadModule(TEXT("GameplayTagsEditor"));
-
 #endif //(WITH_EDITOR && !(UE_BUILD_SHIPPING || UE_BUILD_TEST))
 
 #if WITH_ENGINE
@@ -3247,26 +3244,11 @@ bool FEngineLoop::AppInit( )
 	}
 #endif
 
-#if WITH_ENGINE
-	if (!IsRunningCommandlet())
-	{
-		// Earliest place to init the online subsystems (via plugins)
-		// Code needs GConfigFile to be valid
-		// Must be after FThreadStats::StartThread();
-		// Must be before Render/RHI subsystem D3DCreate()
-		// For platform services that need D3D hooks like Steam
-		// --
-		// Why load HTTP?
-		// Because most, if not all online subsystems will load HTTP themselves. This can cause problems though, as HTTP will be loaded *after* OSS, 
-		// and if OSS holds on to resources allocated by it, this will cause crash (modules are being unloaded in LIFO order with no dependency tracking).
-		// Loading HTTP before OSS works around this problem by making ModuleManager unload HTTP after OSS, at the cost of extra module for the few OSS (like Null) that don't use it.
-
-		// These should not be LoadModuleChecked because these modules might not exist
-		FModuleManager::Get().LoadModule(TEXT("XMPP"));
-		FModuleManager::Get().LoadModule(TEXT("HTTP"));
-		// OSS Default/Console are loaded in plugins immediately following
-	}
-#endif
+	// NOTE: This is the earliest place to init the online subsystems (via plugins)
+	// Code needs GConfigFile to be valid
+	// Must be after FThreadStats::StartThread();
+	// Must be before Render/RHI subsystem D3DCreate()
+	// For platform services that need D3D hooks like Steam
 
 	// Load "pre-init" plugin modules
 	if (!IProjectManager::Get().LoadModulesForProject(ELoadingPhase::PostConfigInit) || !IPluginManager::Get().LoadModulesForEnabledPlugins(ELoadingPhase::PostConfigInit))
diff --git a/Engine/Source/Runtime/LevelSequence/Private/LevelSequenceActor.cpp b/Engine/Source/Runtime/LevelSequence/Private/LevelSequenceActor.cpp
--- a/Engine/Source/Runtime/LevelSequence/Private/LevelSequenceActor.cpp
+++ b/Engine/Source/Runtime/LevelSequence/Private/LevelSequenceActor.cpp
@@ -81,6 +81,8 @@ bool ALevelSequenceActor::GetReferencedContentObjects(TArray<UObject*>& Objects)
 
 void ALevelSequenceActor::Tick(float DeltaSeconds)
 {
+	Super::Tick(DeltaSeconds);
+
 	if (SequencePlayer)
 	{
 		SequencePlayer->Update(DeltaSeconds);
diff --git a/Engine/Source/Runtime/Navmesh/Private/Detour/DetourNavMeshQuery.cpp b/Engine/Source/Runtime/Navmesh/Private/Detour/DetourNavMeshQuery.cpp
--- a/Engine/Source/Runtime/Navmesh/Private/Detour/DetourNavMeshQuery.cpp
+++ b/Engine/Source/Runtime/Navmesh/Private/Detour/DetourNavMeshQuery.cpp
@@ -4101,7 +4101,7 @@ dtStatus dtNavMeshQuery::getPolyWallSegments(dtPolyRef ref, const dtQueryFilter*
 }
 
 static void storeWallSegment(const dtMeshTile* tile, const dtPoly* poly, int edge,
-	dtPolyRef ref0, dtPolyRef ref1, const float* centerPos, const float radiusSqr,
+	dtPolyRef ref0, dtPolyRef ref1, const dtNavMesh* nav, const float* centerPos, const float radiusSqr,
 	float* resultWalls, dtPolyRef* resultRefs, int* resultCount, const int maxResult)
 {
 	if (*resultCount < maxResult)
@@ -4113,12 +4113,80 @@ static void storeWallSegment(const dtMeshTile* tile, const dtPoly* poly, int edg
 		float distSqr = dtDistancePtSegSqr2D(centerPos, va, vb, tseg);
 		if (distSqr <= radiusSqr)
 		{
-			dtVcopy(&resultWalls[*resultCount * 6 + 0], va);
-			dtVcopy(&resultWalls[*resultCount * 6 + 3], vb);
+			const int32 Wall0Offset = (*resultCount * 6) + 0;
+			const int32 Wall1Offset = (*resultCount * 6) + 3;
+
+			dtVcopy(&resultWalls[Wall0Offset], va);
+			dtVcopy(&resultWalls[Wall1Offset], vb);
 			resultRefs[*resultCount * 2 + 0] = ref0;
 			resultRefs[*resultCount * 2 + 1] = ref1;
 
 			*resultCount += 1;
+
+			// find intersection between two edge segments
+			if (nav)
+			{
+				const dtMeshTile* NeiTile = 0;
+				const dtPoly* NeiPoly = 0;
+				nav->getTileAndPolyByRef(ref1, &NeiTile, &NeiPoly);
+
+				unsigned int NeiLinkId = NeiPoly ? NeiPoly->firstLink : DT_NULL_LINK;
+				while (NeiLinkId != DT_NULL_LINK)
+				{
+					const dtLink& link = nav->getLink(NeiTile, NeiLinkId);
+					NeiLinkId = link.next;
+
+					if (link.ref == ref0)
+					{
+						const float* va2 = &NeiTile->verts[NeiPoly->verts[link.edge] * 3];
+						const float* vb2 = &NeiTile->verts[NeiPoly->verts[(link.edge + 1) % NeiPoly->vertCount] * 3];
+
+						// project segment va2-vb2 on va-vb: point va2
+						float seg[3], toPt[3], closestA[3], closestB[3];
+						dtVsub(seg, vb, va);
+
+						dtVsub(toPt, va2, va);
+						float d1 = dtVdot(toPt, seg);
+						float d2 = dtVdot(seg, seg);
+
+						if (d1 <= 0)
+						{
+							dtVcopy(closestA, va);
+						}
+						else if (d2 <= d1)
+						{
+							dtVcopy(closestA, vb);
+						}
+						else
+						{
+							dtVmad(closestA, va, seg, d1 / d2);
+						}
+
+						// project segment va2-vb2 on va-vb: point vb2
+						dtVsub(toPt, vb2, va);
+						d1 = dtVdot(toPt, seg);
+						d2 = dtVdot(seg, seg);
+
+						if (d1 <= 0)
+						{
+							dtVcopy(closestB, va);
+						}
+						else if (d2 <= d1)
+						{
+							dtVcopy(closestB, vb);
+						}
+						else
+						{
+							dtVmad(closestB, va, seg, d1 / d2);
+						}
+
+						// store projected segment (intersection of both edges)
+						dtVcopy(&resultWalls[Wall0Offset], closestA);
+						dtVcopy(&resultWalls[Wall1Offset], closestB);
+						break;
+					}
+				}
+			}
 		}
 	}
 }
@@ -4188,7 +4256,7 @@ dtStatus dtNavMeshQuery::findWallsInNeighbourhood(dtPolyRef startRef, const floa
 			{
 				// store wall segment
 				storeWallSegment(curTile, curPoly, link.edge, 
-					curRef, neighbourRef, centerPos, radiusSqr, 
+					curRef, 0, 0, centerPos, radiusSqr, 
 					resultWalls, resultRefs, resultCount, maxResult);
 				continue;
 			}
@@ -4214,9 +4282,10 @@ dtStatus dtNavMeshQuery::findWallsInNeighbourhood(dtPolyRef startRef, const floa
 			if (!filter->passFilter(neighbourRef, neighbourTile, neighbourPoly) || !passLinkFilterByRef(neighbourTile, neighbourRef))
 			{
 				// store wall segment
-				storeWallSegment(curTile, curPoly, link.edge, 
-					curRef, neighbourRef, centerPos, radiusSqr,
+				storeWallSegment(curTile, curPoly, link.edge,
+					curRef, neighbourRef, m_nav, centerPos, radiusSqr,
 					resultWalls, resultRefs, resultCount, maxResult);
+
 				continue;
 			}
 
@@ -4281,7 +4350,7 @@ dtStatus dtNavMeshQuery::findWallsInNeighbourhood(dtPolyRef startRef, const floa
 			if (bStoreEdge)
 			{
 				storeWallSegment(curTile, curPoly, neighbourIndex,
-					curRef, 0, centerPos, radiusSqr, 
+					curRef, 0, 0, centerPos, radiusSqr, 
 					resultWalls, resultRefs, resultCount, maxResult);
 			}
 		}
diff --git a/Engine/Source/Runtime/Online/HTTP/Private/HttpModule.cpp b/Engine/Source/Runtime/Online/HTTP/Private/HttpModule.cpp
--- a/Engine/Source/Runtime/Online/HTTP/Private/HttpModule.cpp
+++ b/Engine/Source/Runtime/Online/HTTP/Private/HttpModule.cpp
@@ -67,6 +67,10 @@ void FHttpModule::PostLoadCallback()
 }
 
 void FHttpModule::PreUnloadCallback()
+{
+}
+
+void FHttpModule::ShutdownModule()
 {
 	if (HttpManager != nullptr)
 	{
@@ -99,10 +103,6 @@ void FHttpModule::PreUnloadCallback()
 	Singleton = nullptr;
 }
 
-void FHttpModule::ShutdownModule()
-{
-}
-
 bool FHttpModule::HandleHTTPCommand(const TCHAR* Cmd, FOutputDevice& Ar)
 {
 	if (FParse::Command(&Cmd, TEXT("TEST")))
diff --git a/Engine/Source/Runtime/Windows/XAudio2/Private/XAudio2Buffer.cpp b/Engine/Source/Runtime/Windows/XAudio2/Private/XAudio2Buffer.cpp
--- a/Engine/Source/Runtime/Windows/XAudio2/Private/XAudio2Buffer.cpp
+++ b/Engine/Source/Runtime/Windows/XAudio2/Private/XAudio2Buffer.cpp
@@ -345,7 +345,7 @@ bool FXAudio2SoundBuffer::ReadCompressedData( uint8* Destination, bool bLooping
 
 void FXAudio2SoundBuffer::Seek( const float SeekTime )
 {
-	if (ensure(DecompressionState))
+	if (DecompressionState)
 	{
 		DecompressionState->SeekToTime(SeekTime);
 	}
diff --git a/Engine/Source/Runtime/Windows/XAudio2/Private/XAudio2Device.cpp b/Engine/Source/Runtime/Windows/XAudio2/Private/XAudio2Device.cpp
--- a/Engine/Source/Runtime/Windows/XAudio2/Private/XAudio2Device.cpp
+++ b/Engine/Source/Runtime/Windows/XAudio2/Private/XAudio2Device.cpp
@@ -278,26 +278,6 @@ void FXAudio2Device::UpdateHardware()
 
 void FXAudio2Device::CheckDeviceStateChange()
 {
-#if PLATFORM_WINDOWS
-	if (DeviceProperties)
-	{
-		if (DeviceProperties->DidAudioDeviceChange())
-		{
-			// Stop any sounds that are playing
-			StopAllSounds(true);
-
-			// Set all sound sources to virtual mode so they don't play audio
-			for (FSoundSource* Source : Sources)
-			{
-				Source->SetVirtual();
-			}
-
-			// And switch to no-audio mode.
-			bIsAudioDeviceHardwareInitialized = false;
-		}
-	}
-#endif
-
 	// Update the audio clock time
 	AudioClock = DeviceProperties->GetAudioClockTime();
 }
diff --git a/Engine/Source/Runtime/Windows/XAudio2/Private/XAudio2Source.cpp b/Engine/Source/Runtime/Windows/XAudio2/Private/XAudio2Source.cpp
--- a/Engine/Source/Runtime/Windows/XAudio2/Private/XAudio2Source.cpp
+++ b/Engine/Source/Runtime/Windows/XAudio2/Private/XAudio2Source.cpp
@@ -69,7 +69,7 @@ FXAudio2SoundSource::FXAudio2SoundSource(FAudioDevice* InAudioDevice)
 	if (!InAudioDevice->bIsAudioDeviceHardwareInitialized)
 	{
 		bIsVirtual = true;
-}
+	}
 }
 
 /**
@@ -120,7 +120,8 @@ void FXAudio2SoundSource::FreeResources( void )
 
 	if (bResourcesNeedFreeing && Buffer)
 	{
-		check(Buffer->ResourceID == 0);
+		// If we failed to initialize, then we will have a non-zero resource ID, but still need to delete the buffer
+		check(!bInitialized || Buffer->ResourceID == 0);
 		delete Buffer;
 	}
 
@@ -457,6 +458,12 @@ bool FXAudio2SoundSource::CreateSource( void )
 
 bool FXAudio2SoundSource::PrepareForInitialization(FWaveInstance* InWaveInstance)
 {
+	// If the headphones have been unplugged, set this voice to be virtual
+	if (!AudioDevice->DeviceProperties->bAllowNewVoices)
+	{
+		bIsVirtual = true;
+	}
+
 	// If virtual only need wave instance data and no need to load source data
 	if (bIsVirtual)
 	{
@@ -523,13 +530,13 @@ bool FXAudio2SoundSource::IsPreparedToInit()
 
 bool FXAudio2SoundSource::Init(FWaveInstance* InWaveInstance)
 {
+	// Setup our virtual duration/playback data regardless of it's currently virtual since it can flip while playing
+	VirtualDuration = InWaveInstance->WaveData->GetDuration();
+	VirtualPlaybackTime = 0.0f;
+
 	if (bIsVirtual)
 	{
 		bInitialized = true;
-
-		// Setup our virtual duration/playback data
-		VirtualDuration = InWaveInstance->WaveData->GetDuration();
-		VirtualPlaybackTime = 0.0f;
 		return true;
 	}
 
@@ -597,10 +604,19 @@ bool FXAudio2SoundSource::Init(FWaveInstance* InWaveInstance)
 			// Initialization succeeded.
 			return true;
 		}
-		else
+		else if (AudioDevice->DeviceProperties->bAllowNewVoices)
 		{
 			UE_LOG(LogXAudio2, Warning, TEXT("Failed to init sound source for wave instance '%s' due to being unable to create an IXAudio2SourceVoice."), *InWaveInstance->GetName());
 		}
+		else
+		{
+			// The audio device was unplugged after init was declared, we're actually virtual now
+			bIsVirtual = true;
+
+			// Set that we've actually successfully initialized
+			bInitialized = true;
+			return true;
+		}
 	}
 	else
 	{
@@ -1312,15 +1328,21 @@ void FXAudio2SoundSource::Update()
 
 	Pitch = FMath::Clamp<float>(Pitch, MIN_PITCH, MAX_PITCH);
 
-	// If this is a virtual source, then update it's duration and do any notification on completion based on duration
-	if (bIsVirtual)
+	// If the headphones have been unplugged after playing, set this voice to be virtual
+	if (!AudioDevice->DeviceProperties->bAllowNewVoices)
 	{
-		// Get the game-thread update delta time
-		float DeviceDeltaTime = AudioDevice->GetUpdateDeltaTime();
+		bIsVirtual = true;
+	}
+
+	// Track virtual playback time even if the voice is not virtual, it can flip to being virtual while playing.
+	const float DeviceDeltaTime = AudioDevice->GetUpdateDeltaTime();
 
-		// Scale the virtual playback time based on the pitch of the sound
-		VirtualPlaybackTime += DeviceDeltaTime * Pitch;
+	// Scale the virtual playback time based on the pitch of the sound
+	VirtualPlaybackTime += DeviceDeltaTime * Pitch;
 
+	// If this is a virtual source, then update it's duration and do any notification on completion based on duration
+	if (bIsVirtual)
+	{
 		if (VirtualPlaybackTime >= VirtualDuration)
 		{
 			if (WaveInstance->LoopingMode == LOOP_Never)
