Fix for UGameplayTagManager::OnLastChanceToAddNativeTags not returning delegate by reference.

#tests compile

diff --git a/Engine/Source/Runtime/GameplayTags/Classes/GameplayTagsManager.h b/Engine/Source/Runtime/GameplayTags/Classes/GameplayTagsManager.h
--- a/Engine/Source/Runtime/GameplayTags/Classes/GameplayTagsManager.h
+++ b/Engine/Source/Runtime/GameplayTags/Classes/GameplayTagsManager.h
@@ -227,7 +227,7 @@ class GAMEPLAYTAGS_API UGameplayTagsManager : public UObject
 	/** Call to flush the list of native tags, once called it is unsafe to add more */
 	void DoneAddingNativeTags();
 
-	static FSimpleMulticastDelegate OnLastChanceToAddNativeTags();
+	static FSimpleMulticastDelegate& OnLastChanceToAddNativeTags();
 
 	/**
 	 * Gets a Tag Container containing the supplied tag and all of it's parents as explicit tags
diff --git a/Engine/Source/Runtime/GameplayTags/Private/GameplayTagsManager.cpp b/Engine/Source/Runtime/GameplayTags/Private/GameplayTagsManager.cpp
--- a/Engine/Source/Runtime/GameplayTags/Private/GameplayTagsManager.cpp
+++ b/Engine/Source/Runtime/GameplayTags/Private/GameplayTagsManager.cpp
@@ -1156,7 +1156,7 @@ FGameplayTag UGameplayTagsManager::AddNativeGameplayTag(FName TagName)
 	return FGameplayTag();
 }
 
-FSimpleMulticastDelegate UGameplayTagsManager::OnLastChanceToAddNativeTags()
+FSimpleMulticastDelegate& UGameplayTagsManager::OnLastChanceToAddNativeTags()
 {
 	static FSimpleMulticastDelegate Delegate;
 	return Delegate;
