Fix for non deterministic GameplayTag Net Indices

diff --git a/Engine/Source/Runtime/GameplayTags/Private/GameplayTagsManager.cpp b/Engine/Source/Runtime/GameplayTags/Private/GameplayTagsManager.cpp
--- a/Engine/Source/Runtime/GameplayTags/Private/GameplayTagsManager.cpp
+++ b/Engine/Source/Runtime/GameplayTags/Private/GameplayTagsManager.cpp
@@ -70,7 +70,10 @@ struct FCompareFGameplayTagNodeByTag
 {
 	FORCEINLINE bool operator()( const TSharedPtr<FGameplayTagNode>& A, const TSharedPtr<FGameplayTagNode>& B ) const
 	{
-		return (A->GetSimpleTagName().Compare(B->GetSimpleTagName())) < 0;
+		// Note: GetSimpleTagName() is not good enough here. The individual tag nodes are share frequently (E.g, Dog.Tail, Cat.Tail have sub nodes with the same simple tag name)
+		// Compare with equal FNames will look at the backing number/indice to the FName. For FNames used elsewhere, like "A" for example, this can cause non determinism in platforms
+		// (For example if static order initialization differs on two platforms, the "version" of the "A" FName that two places get could be different, causing this comparison to also be)
+		return (A->GetCompleteTagName().Compare(B->GetCompleteTagName())) < 0;
 	}
 };
 
