Mark various operators as BlueprintThreadSafe for less common gameplay types that might want to be used by animation

diff --git a/Engine/Source/Runtime/Engine/Classes/Kismet/KismetGuidLibrary.h b/Engine/Source/Runtime/Engine/Classes/Kismet/KismetGuidLibrary.h
--- a/Engine/Source/Runtime/Engine/Classes/Kismet/KismetGuidLibrary.h
+++ b/Engine/Source/Runtime/Engine/Classes/Kismet/KismetGuidLibrary.h
@@ -16,19 +16,19 @@ class ENGINE_API UKismetGuidLibrary
 	GENERATED_UCLASS_BODY()
 
 	/** Returns true if the values are equal (A == B) */
-	UFUNCTION(BlueprintPure, meta=(DisplayName="Equal (Guid)", CompactNodeTitle="=="), Category="Guid")
+	UFUNCTION(BlueprintPure, meta=(DisplayName="Equal (Guid)", CompactNodeTitle="==", BlueprintThreadSafe), Category="Guid")
 	static bool EqualEqual_GuidGuid( const FGuid& A, const FGuid& B );
 	
 	/** Returns true if the values are not equal (A != B) */
-	UFUNCTION(BlueprintPure, meta=(DisplayName="NotEqual (Guid)", CompactNodeTitle="!="), Category="Guid")
+	UFUNCTION(BlueprintPure, meta=(DisplayName="NotEqual (Guid)", CompactNodeTitle="!=", BlueprintThreadSafe), Category="Guid")
 	static bool NotEqual_GuidGuid( const FGuid& A, const FGuid& B );
 
 	/** Checks whether the given GUID is valid */
-	UFUNCTION(BlueprintPure, meta=(DisplayName="IsValid (Guid)", CompactNodeTitle="isValid?"), Category="Guid")
+	UFUNCTION(BlueprintPure, meta=(DisplayName="IsValid (Guid)", CompactNodeTitle="isValid?", BlueprintThreadSafe), Category="Guid")
 	static bool IsValid_Guid( const FGuid& InGuid );
 
 	/** Invalidates the given GUID */
-	UFUNCTION(BlueprintCallable, meta=(DisplayName="Invalidate_Guid (Guid)"), Category="Guid")
+	UFUNCTION(BlueprintCallable, meta=(DisplayName="Invalidate_Guid (Guid)", BlueprintThreadSafe), Category="Guid")
 	static void Invalidate_Guid( UPARAM(ref) FGuid& InGuid );
 
 	/** Returns a new unique GUID */
@@ -36,10 +36,10 @@ class ENGINE_API UKismetGuidLibrary
 	static FGuid NewGuid();
 	
 	/** Converts a GUID value to a string, in the form 'A-B-C-D' */
-	UFUNCTION(BlueprintPure, meta=(DisplayName="ToString (Guid)", CompactNodeTitle="->"), Category="Guid")
+	UFUNCTION(BlueprintPure, meta=(DisplayName="ToString (Guid)", CompactNodeTitle="->", BlueprintThreadSafe), Category="Guid")
 	static FString Conv_GuidToString( const FGuid& InGuid );
 
 	/** Converts a String of format EGuidFormats to a Guid. Returns Guid OutGuid, Returns bool Success */
-	UFUNCTION(BlueprintPure, meta=(DisplayName="Parse String to Guid"), Category="Guid")
+	UFUNCTION(BlueprintPure, meta=(DisplayName="Parse String to Guid", BlueprintThreadSafe), Category="Guid")
 	static void Parse_StringToGuid( const FString& GuidString, FGuid& OutGuid, bool& Success );
 };
diff --git a/Engine/Source/Runtime/Engine/Classes/Kismet/KismetSystemLibrary.h b/Engine/Source/Runtime/Engine/Classes/Kismet/KismetSystemLibrary.h
--- a/Engine/Source/Runtime/Engine/Classes/Kismet/KismetSystemLibrary.h
+++ b/Engine/Source/Runtime/Engine/Classes/Kismet/KismetSystemLibrary.h
@@ -185,51 +185,51 @@ class ENGINE_API UKismetSystemLibrary : public UBlueprintFunctionLibrary
 	static UObject* Conv_InterfaceToObject(const FScriptInterface& Interface); 
 
 	/** Builds a SoftObjectPath struct. Generally you should be using Soft Object References/Ptr types instead */
-	UFUNCTION(BlueprintPure, Category = "SoftObjectPath", meta = (Keywords = "construct build", NativeMakeFunc))
+	UFUNCTION(BlueprintPure, Category = "SoftObjectPath", meta = (Keywords = "construct build", NativeMakeFunc, BlueprintThreadSafe))
 	static FSoftObjectPath MakeSoftObjectPath(const FString& PathString);
 
 	/** Gets the path string out of a Soft Object Path */
-	UFUNCTION(BlueprintPure, Category = "SoftObjectPath", meta = (NativeBreakFunc))
+	UFUNCTION(BlueprintPure, Category = "SoftObjectPath", meta = (NativeBreakFunc, BlueprintThreadSafe))
 	static void BreakSoftObjectPath(FSoftObjectPath InSoftObjectPath, FString& PathString);
 
 	/** Builds a SoftClassPath struct. Generally you should be using Soft Class References/Ptr types instead */
-	UFUNCTION(BlueprintPure, Category = "SoftClassPath", meta = (Keywords = "construct build", NativeMakeFunc))
+	UFUNCTION(BlueprintPure, Category = "SoftClassPath", meta = (Keywords = "construct build", NativeMakeFunc, BlueprintThreadSafe))
 	static FSoftClassPath MakeSoftClassPath(const FString& PathString);
 
 	/** Gets the path string out of a Soft Class Path */
-	UFUNCTION(BlueprintPure, Category = "SoftClassPath", meta = (NativeBreakFunc))
+	UFUNCTION(BlueprintPure, Category = "SoftClassPath", meta = (NativeBreakFunc, BlueprintThreadSafe))
 	static void BreakSoftClassPath(FSoftClassPath InSoftClassPath, FString& PathString);
 
 	/** Returns true if the Soft Object Reference is not null */
-	UFUNCTION(BlueprintPure, Category = "Utilities")
+	UFUNCTION(BlueprintPure, Category = "Utilities", meta = (BlueprintThreadSafe))
 	static bool IsValidSoftObjectReference(const TSoftObjectPtr<UObject>& SoftObjectReference);
 
 	/** Converts a Soft Object Reference to a string. The other direction is not provided because it cannot be validated */
-	UFUNCTION(BlueprintPure, meta = (DisplayName = "ToString (SoftObjectReference)", CompactNodeTitle = "->"), Category = "Utilities")
+	UFUNCTION(BlueprintPure, meta = (DisplayName = "ToString (SoftObjectReference)", CompactNodeTitle = "->", BlueprintThreadSafe), Category = "Utilities")
 	static FString Conv_SoftObjectReferenceToString(const TSoftObjectPtr<UObject>& SoftObjectReference);
 
 	/** Returns true if the values are equal (A == B) */
-	UFUNCTION(BlueprintPure, meta = (DisplayName = "Equal (SoftObjectReference)", CompactNodeTitle = "=="), Category = "Utilities")
+	UFUNCTION(BlueprintPure, meta = (DisplayName = "Equal (SoftObjectReference)", CompactNodeTitle = "==", BlueprintThreadSafe), Category = "Utilities")
 	static bool EqualEqual_SoftObjectReference(const TSoftObjectPtr<UObject>& A, const TSoftObjectPtr<UObject>& B);
 
 	/** Returns true if the values are not equal (A != B) */
-	UFUNCTION(BlueprintPure, meta = (DisplayName = "NotEqual (SoftObjectReference)", CompactNodeTitle = "!="), Category = "Utilities")
+	UFUNCTION(BlueprintPure, meta = (DisplayName = "NotEqual (SoftObjectReference)", CompactNodeTitle = "!=", BlueprintThreadSafe), Category = "Utilities")
 	static bool NotEqual_SoftObjectReference(const TSoftObjectPtr<UObject>& A, const TSoftObjectPtr<UObject>& B);
 
 	/** Returns true if the Soft Class Reference is not null */
-	UFUNCTION(BlueprintPure, Category = "Utilities")
+	UFUNCTION(BlueprintPure, Category = "Utilities", meta = (BlueprintThreadSafe))
 	static bool IsValidSoftClassReference(const TSoftClassPtr<UObject>& SoftClassReference);
 
 	/** Converts a Soft Class Reference to a string. The other direction is not provided because it cannot be validated */
-	UFUNCTION(BlueprintPure, meta = (DisplayName = "ToString (SoftObjectReference)", CompactNodeTitle = "->"), Category = "Utilities")
+	UFUNCTION(BlueprintPure, meta = (DisplayName = "ToString (SoftObjectReference)", CompactNodeTitle = "->", BlueprintThreadSafe), Category = "Utilities")
 	static FString Conv_SoftClassReferenceToString(const TSoftClassPtr<UObject>& SoftClassReference);
 
 	/** Returns true if the values are equal (A == B) */
-	UFUNCTION(BlueprintPure, meta = (DisplayName = "Equal (SoftClassReference)", CompactNodeTitle = "=="), Category = "Utilities")
+	UFUNCTION(BlueprintPure, meta = (DisplayName = "Equal (SoftClassReference)", CompactNodeTitle = "==", BlueprintThreadSafe), Category = "Utilities")
 	static bool EqualEqual_SoftClassReference(const TSoftClassPtr<UObject>& A, const TSoftClassPtr<UObject>& B);
 
 	/** Returns true if the values are not equal (A != B) */
-	UFUNCTION(BlueprintPure, meta = (DisplayName = "NotEqual (SoftClassReference)", CompactNodeTitle = "!="), Category = "Utilities")
+	UFUNCTION(BlueprintPure, meta = (DisplayName = "NotEqual (SoftClassReference)", CompactNodeTitle = "!=", BlueprintThreadSafe), Category = "Utilities")
 	static bool NotEqual_SoftClassReference(const TSoftClassPtr<UObject>& A, const TSoftClassPtr<UObject>& B);
 
 	UFUNCTION(BlueprintPure, meta = (BlueprintInternalUseOnly = "true"), Category = "Utilities")
@@ -1687,35 +1687,35 @@ class ENGINE_API UKismetSystemLibrary : public UBlueprintFunctionLibrary
 	static void GetPrimaryAssetIdList(FPrimaryAssetType PrimaryAssetType, TArray<FPrimaryAssetId>& OutPrimaryAssetIdList);
 
 	/** Returns true if the Primary Asset Id is valid */
-	UFUNCTION(BlueprintPure, Category = "AssetManager", meta=(ScriptMethod="IsValid", ScriptOperator="bool"))
+	UFUNCTION(BlueprintPure, Category = "AssetManager", meta=(ScriptMethod="IsValid", ScriptOperator="bool", BlueprintThreadSafe))
 	static bool IsValidPrimaryAssetId(FPrimaryAssetId PrimaryAssetId);
 
 	/** Converts a Primary Asset Id to a string. The other direction is not provided because it cannot be validated */
-	UFUNCTION(BlueprintPure, meta = (DisplayName = "ToString (PrimaryAssetId)", CompactNodeTitle = "->", ScriptMethod="ToString"), Category = "AssetManager")
+	UFUNCTION(BlueprintPure, meta = (DisplayName = "ToString (PrimaryAssetId)", CompactNodeTitle = "->", ScriptMethod="ToString", BlueprintThreadSafe), Category = "AssetManager")
 	static FString Conv_PrimaryAssetIdToString(FPrimaryAssetId PrimaryAssetId);
 
 	/** Returns true if the values are equal (A == B) */
-	UFUNCTION(BlueprintPure, meta = (DisplayName = "Equal (PrimaryAssetId)", CompactNodeTitle = "==", ScriptOperator="=="), Category = "AssetManager")
+	UFUNCTION(BlueprintPure, meta = (DisplayName = "Equal (PrimaryAssetId)", CompactNodeTitle = "==", ScriptOperator="==", BlueprintThreadSafe), Category = "AssetManager")
 	static bool EqualEqual_PrimaryAssetId(FPrimaryAssetId A, FPrimaryAssetId B);
 
 	/** Returns true if the values are not equal (A != B) */
-	UFUNCTION(BlueprintPure, meta = (DisplayName = "NotEqual (PrimaryAssetId)", CompactNodeTitle = "!=", ScriptOperator="!="), Category = "AssetManager")
+	UFUNCTION(BlueprintPure, meta = (DisplayName = "NotEqual (PrimaryAssetId)", CompactNodeTitle = "!=", ScriptOperator="!=", BlueprintThreadSafe), Category = "AssetManager")
 	static bool NotEqual_PrimaryAssetId(FPrimaryAssetId A, FPrimaryAssetId B);
 
 	/** Returns list of Primary Asset Ids for a PrimaryAssetType */
-	UFUNCTION(BlueprintPure, Category = "AssetManager", meta=(ScriptMethod="IsValid", ScriptOperator="bool"))
+	UFUNCTION(BlueprintPure, Category = "AssetManager", meta=(ScriptMethod="IsValid", ScriptOperator="bool", BlueprintThreadSafe))
 	static bool IsValidPrimaryAssetType(FPrimaryAssetType PrimaryAssetType);
 
 	/** Converts a Primary Asset Type to a string. The other direction is not provided because it cannot be validated */
-	UFUNCTION(BlueprintPure, meta = (DisplayName = "ToString (PrimaryAssetType)", CompactNodeTitle = "->", ScriptMethod="ToString"), Category = "AssetManager")
+	UFUNCTION(BlueprintPure, meta = (DisplayName = "ToString (PrimaryAssetType)", CompactNodeTitle = "->", ScriptMethod="ToString", BlueprintThreadSafe), Category = "AssetManager")
 	static FString Conv_PrimaryAssetTypeToString(FPrimaryAssetType PrimaryAssetType);
 
 	/** Returns true if the values are equal (A == B) */
-	UFUNCTION(BlueprintPure, meta = (DisplayName = "Equal (PrimaryAssetType)", CompactNodeTitle = "==", ScriptOperator="=="), Category = "AssetManager")
+	UFUNCTION(BlueprintPure, meta = (DisplayName = "Equal (PrimaryAssetType)", CompactNodeTitle = "==", ScriptOperator="==", BlueprintThreadSafe), Category = "AssetManager")
 	static bool EqualEqual_PrimaryAssetType(FPrimaryAssetType A, FPrimaryAssetType B);
 
 	/** Returns true if the values are not equal (A != B) */
-	UFUNCTION(BlueprintPure, meta = (DisplayName = "NotEqual (PrimaryAssetType)", CompactNodeTitle = "!=", ScriptOperator="!="), Category = "AssetManager")
+	UFUNCTION(BlueprintPure, meta = (DisplayName = "NotEqual (PrimaryAssetType)", CompactNodeTitle = "!=", ScriptOperator="!=", BlueprintThreadSafe), Category = "AssetManager")
 	static bool NotEqual_PrimaryAssetType(FPrimaryAssetType A, FPrimaryAssetType B);
 
 	/** Unloads a primary asset, which allows it to be garbage collected if nothing else is referencing it */
diff --git a/Engine/Source/Runtime/GameplayTags/Classes/BlueprintGameplayTagLibrary.h b/Engine/Source/Runtime/GameplayTags/Classes/BlueprintGameplayTagLibrary.h
--- a/Engine/Source/Runtime/GameplayTags/Classes/BlueprintGameplayTagLibrary.h
+++ b/Engine/Source/Runtime/GameplayTags/Classes/BlueprintGameplayTagLibrary.h
@@ -25,7 +25,7 @@ class UBlueprintGameplayTagLibrary : public UBlueprintFunctionLibrary
 	 * 
 	 * @return True if TagOne matches TagTwo
 	 */
-	UFUNCTION(BlueprintPure, Category="GameplayTags", meta = (Keywords = "DoGameplayTagsMatch"))
+	UFUNCTION(BlueprintPure, Category="GameplayTags", meta = (Keywords = "DoGameplayTagsMatch", BlueprintThreadSafe))
 	static bool MatchesTag(FGameplayTag TagOne, FGameplayTag TagTwo, bool bExactMatch);
 
 	/**
@@ -37,27 +37,27 @@ class UBlueprintGameplayTagLibrary : public UBlueprintFunctionLibrary
 	 * 
 	 * @return True if TagOne matches any tags explicitly present in OtherContainer
 	 */
-	UFUNCTION(BlueprintPure, Category="GameplayTags")
+	UFUNCTION(BlueprintPure, Category="GameplayTags", meta = (BlueprintThreadSafe))
 	static bool MatchesAnyTags(FGameplayTag TagOne, const FGameplayTagContainer& OtherContainer, bool bExactMatch);
 
 	/** Returns true if the values are equal (A == B) */
-	UFUNCTION(BlueprintPure, meta=(DisplayName="Equal (GameplayTag)", CompactNodeTitle="=="), Category="GameplayTags")
+	UFUNCTION(BlueprintPure, meta=(DisplayName="Equal (GameplayTag)", CompactNodeTitle="==", BlueprintThreadSafe), Category="GameplayTags")
 	static bool EqualEqual_GameplayTag( FGameplayTag A, FGameplayTag B );
 	
 	/** Returns true if the values are not equal (A != B) */
-	UFUNCTION(BlueprintPure, meta=(DisplayName="NotEqual (GameplayTag)", CompactNodeTitle="!="), Category="GameplayTags")
+	UFUNCTION(BlueprintPure, meta=(DisplayName="NotEqual (GameplayTag)", CompactNodeTitle="!=", BlueprintThreadSafe), Category="GameplayTags")
 	static bool NotEqual_GameplayTag( FGameplayTag A, FGameplayTag B );
 
 	/** Returns true if the passed in gameplay tag is non-null */
-	UFUNCTION(BlueprintPure, Category = "GameplayTags")
+	UFUNCTION(BlueprintPure, Category = "GameplayTags", meta = (BlueprintThreadSafe))
 	static bool IsGameplayTagValid(FGameplayTag GameplayTag);
 
 	/** Returns FName of this tag */
-	UFUNCTION(BlueprintPure, Category = "GameplayTags")
+	UFUNCTION(BlueprintPure, Category = "GameplayTags", meta = (BlueprintThreadSafe))
 	static FName GetTagName(const FGameplayTag& GameplayTag);
 
 	/** Creates a literal FGameplayTag */
-	UFUNCTION(BlueprintPure, Category = "GameplayTags")
+	UFUNCTION(BlueprintPure, Category = "GameplayTags", meta = (BlueprintThreadSafe))
 	static FGameplayTag MakeLiteralGameplayTag(FGameplayTag Value);
 
 	/**
@@ -67,7 +67,7 @@ class UBlueprintGameplayTagLibrary : public UBlueprintFunctionLibrary
 	 * 
 	 * @return The number of tags in the specified container
 	 */
-	UFUNCTION(BlueprintPure, Category="GameplayTags")
+	UFUNCTION(BlueprintPure, Category="GameplayTags", meta = (BlueprintThreadSafe))
 	static int32 GetNumGameplayTagsInContainer(const FGameplayTagContainer& TagContainer);
 
 	/**
@@ -79,7 +79,7 @@ class UBlueprintGameplayTagLibrary : public UBlueprintFunctionLibrary
 	 *
 	 * @return True if the container has the specified tag, false if it does not
 	 */
-	UFUNCTION(BlueprintPure, Category="GameplayTags", meta = (Keywords = "DoesContainerHaveTag"))
+	UFUNCTION(BlueprintPure, Category="GameplayTags", meta = (Keywords = "DoesContainerHaveTag", BlueprintThreadSafe))
 	static bool HasTag(const FGameplayTagContainer& TagContainer, FGameplayTag Tag, bool bExactMatch);
 
 	/**
@@ -91,7 +91,7 @@ class UBlueprintGameplayTagLibrary : public UBlueprintFunctionLibrary
 	 * 
 	 * @return True if the container has ANY of the tags in the other container
 	 */
-	UFUNCTION(BlueprintPure, Category="GameplayTags", meta = (Keywords = "DoesContainerMatchAnyTagsInContainer"))
+	UFUNCTION(BlueprintPure, Category="GameplayTags", meta = (Keywords = "DoesContainerMatchAnyTagsInContainer", BlueprintThreadSafe))
 	static bool HasAnyTags(const FGameplayTagContainer& TagContainer, const FGameplayTagContainer& OtherContainer, bool bExactMatch);
 
 	/**
@@ -103,7 +103,7 @@ class UBlueprintGameplayTagLibrary : public UBlueprintFunctionLibrary
 	 * 
 	 * @return True if the container has ALL of the tags in the other container
 	 */
-	UFUNCTION(BlueprintPure, Category="GameplayTags", meta = (Keywords = "DoesContainerMatchAllTagsInContainer"))
+	UFUNCTION(BlueprintPure, Category="GameplayTags", meta = (Keywords = "DoesContainerMatchAllTagsInContainer", BlueprintThreadSafe))
 	static bool HasAllTags(const FGameplayTagContainer& TagContainer, const FGameplayTagContainer& OtherContainer, bool bExactMatch);
 
 	/**
@@ -114,7 +114,7 @@ class UBlueprintGameplayTagLibrary : public UBlueprintFunctionLibrary
 	 * 
 	 * @return True if the container matches the query, false otherwise.
 	 */
-	UFUNCTION(BlueprintPure, Category = "GameplayTags")
+	UFUNCTION(BlueprintPure, Category = "GameplayTags", meta = (BlueprintThreadSafe))
 	static bool DoesContainerMatchTagQuery(const FGameplayTagContainer& TagContainer, const FGameplayTagQuery& TagQuery);
 
 	/**
@@ -133,7 +133,7 @@ class UBlueprintGameplayTagLibrary : public UBlueprintFunctionLibrary
 	 * @param InOutTagContainer		The container that will be appended too.
 	 * @param Tag					The tag to add to the container
 	 */
-	UFUNCTION(BlueprintCallable, Category = "GameplayTags")
+	UFUNCTION(BlueprintCallable, Category = "GameplayTags", meta = (BlueprintThreadSafe))
 	static void AddGameplayTag(UPARAM(ref) FGameplayTagContainer& TagContainer, FGameplayTag Tag);
 
 	/**
@@ -142,7 +142,7 @@ class UBlueprintGameplayTagLibrary : public UBlueprintFunctionLibrary
 	 * @param InOutTagContainer		The container that will be appended too.
 	 * @param Tag					The tag to add to the container
 	 */
-	UFUNCTION(BlueprintCallable, Category = "GameplayTags")
+	UFUNCTION(BlueprintCallable, Category = "GameplayTags", meta = (BlueprintThreadSafe))
 	static bool RemoveGameplayTag(UPARAM(ref) FGameplayTagContainer& TagContainer, FGameplayTag Tag);
 
 	/**
@@ -151,31 +151,31 @@ class UBlueprintGameplayTagLibrary : public UBlueprintFunctionLibrary
 	 * @param InOutTagContainer		The container that will be appended too.
 	 * @param InTagContainer		The container to append.
 	 */
-	UFUNCTION(BlueprintCallable, Category = "GameplayTags")
+	UFUNCTION(BlueprintCallable, Category = "GameplayTags", meta = (BlueprintThreadSafe))
 	static void AppendGameplayTagContainers(UPARAM(ref) FGameplayTagContainer& InOutTagContainer, const FGameplayTagContainer& InTagContainer);
 
 	/** Returns true if the values are equal (A == B) */
-	UFUNCTION(BlueprintPure, meta=(DisplayName="Equal (GameplayTagContainer)", CompactNodeTitle="=="), Category="GameplayTags")
+	UFUNCTION(BlueprintPure, meta=(DisplayName="Equal (GameplayTagContainer)", CompactNodeTitle="==", BlueprintThreadSafe), Category="GameplayTags")
 	static bool EqualEqual_GameplayTagContainer( const FGameplayTagContainer& A, const FGameplayTagContainer& B );
 	
 	/** Returns true if the values are not equal (A != B) */
-	UFUNCTION(BlueprintPure, meta=(DisplayName="NotEqual (GameplayTagContainer)", CompactNodeTitle="!="), Category="GameplayTags")
+	UFUNCTION(BlueprintPure, meta=(DisplayName="NotEqual (GameplayTagContainer)", CompactNodeTitle="!=", BlueprintThreadSafe), Category="GameplayTags")
 	static bool NotEqual_GameplayTagContainer( const FGameplayTagContainer& A, const FGameplayTagContainer& B );
 
 	/** Creates a literal FGameplayTagContainer */
-	UFUNCTION(BlueprintPure, Category = "GameplayTags")
+	UFUNCTION(BlueprintPure, Category = "GameplayTags", meta = (BlueprintThreadSafe))
 	static FGameplayTagContainer MakeLiteralGameplayTagContainer(FGameplayTagContainer Value);
 
 	/** Creates a FGameplayTagContainer from the array of passed in tags */
-	UFUNCTION(BlueprintPure, Category = "GameplayTags")
+	UFUNCTION(BlueprintPure, Category = "GameplayTags", meta = (BlueprintThreadSafe))
 	static FGameplayTagContainer MakeGameplayTagContainerFromArray(const TArray<FGameplayTag>& GameplayTags);
 
 	/** Creates a FGameplayTagContainer containing a single tag */
-	UFUNCTION(BlueprintPure, Category = "GameplayTags")
+	UFUNCTION(BlueprintPure, Category = "GameplayTags", meta = (BlueprintThreadSafe))
 	static FGameplayTagContainer MakeGameplayTagContainerFromTag(FGameplayTag SingleTag);
 
 	/** Breaks tag container into explicit array of tags */
-	UFUNCTION(BlueprintPure, Category = "GameplayTags")
+	UFUNCTION(BlueprintPure, Category = "GameplayTags", meta = (BlueprintThreadSafe))
 	static void BreakGameplayTagContainer(const FGameplayTagContainer& GameplayTagContainer, TArray<FGameplayTag>& GameplayTags);
 
 	/**
@@ -185,7 +185,7 @@ class UBlueprintGameplayTagLibrary : public UBlueprintFunctionLibrary
 	 *
 	 * @return	The literal FGameplayTagQuery
 	 */
-	UFUNCTION(BlueprintPure, Category = "GameplayTags")
+	UFUNCTION(BlueprintPure, Category = "GameplayTags", meta = (BlueprintThreadSafe))
 	static FGameplayTagQuery MakeGameplayTagQuery(FGameplayTagQuery TagQuery);
 
 	/**
@@ -223,7 +223,7 @@ class UBlueprintGameplayTagLibrary : public UBlueprintFunctionLibrary
 	 *
 	 * @param TagContainer	The tag container to get the debug string from.
 	 */
-	UFUNCTION(BlueprintPure, Category = "GameplayTags")
+	UFUNCTION(BlueprintPure, Category = "GameplayTags", meta = (BlueprintThreadSafe))
 	static FString GetDebugStringFromGameplayTagContainer(const FGameplayTagContainer& TagContainer);
 
 	/**
@@ -231,7 +231,7 @@ class UBlueprintGameplayTagLibrary : public UBlueprintFunctionLibrary
 	 *
 	 * @param GameplayTag	The tag to get the debug string from.
 	 */
-	UFUNCTION(BlueprintPure, Category = "GameplayTags")
+	UFUNCTION(BlueprintPure, Category = "GameplayTags", meta = (BlueprintThreadSafe))
 	static FString GetDebugStringFromGameplayTag(FGameplayTag GameplayTag);
 
 };
