Compute a hash of gameplay tags when using the fast replication path

diff --git a/Engine/Source/Runtime/GameplayTags/Classes/GameplayTagsManager.h b/Engine/Source/Runtime/GameplayTags/Classes/GameplayTagsManager.h
--- a/Engine/Source/Runtime/GameplayTags/Classes/GameplayTagsManager.h
+++ b/Engine/Source/Runtime/GameplayTags/Classes/GameplayTagsManager.h
@@ -474,6 +474,9 @@ class GAMEPLAYTAGS_API UGameplayTagsManager : public UObject
 		return bUseFastReplication;
 	}
 
+	/** Returns the hash of NetworkGameplayTagNodeIndex */
+	uint32 GetNetworkGameplayTagNodeIndexHash() const {	return NetworkGameplayTagNodeIndexHash;	}
+
 	/** Returns a list of the ini files that contain restricted tags */
 	void GetRestrictedTagConfigFiles(TArray<FString>& RestrictedConfigFiles) const;
 
@@ -713,6 +716,8 @@ private:
 	/** Sorted list of nodes, used for network replication */
 	TArray<TSharedPtr<FGameplayTagNode>> NetworkGameplayTagNodeIndex;
 
+	uint32 NetworkGameplayTagNodeIndexHash;
+
 	/** Holds all of the valid gameplay-related tags that can be applied to assets */
 	UPROPERTY()
 	TArray<UDataTable*> GameplayTagTables;
diff --git a/Engine/Source/Runtime/GameplayTags/Private/GameplayTagsManager.cpp b/Engine/Source/Runtime/GameplayTags/Private/GameplayTagsManager.cpp
--- a/Engine/Source/Runtime/GameplayTags/Private/GameplayTagsManager.cpp
+++ b/Engine/Source/Runtime/GameplayTags/Private/GameplayTagsManager.cpp
@@ -43,6 +43,7 @@ UGameplayTagsManager::UGameplayTagsManager(const FObjectInitializer& ObjectIniti
 	NetIndexFirstBitSegment = 16;
 	NetIndexTrueBitNum = 16;
 	NumBitsForContainerSize = 6;
+	NetworkGameplayTagNodeIndexHash = 0;
 }
 
 // Enable to turn on detailed startup logging
@@ -478,12 +479,16 @@ void UGameplayTagsManager::ConstructNetIndex()
 
 	UE_CLOG(PrintNetIndiceAssignment, LogGameplayTags, Display, TEXT("Assigning NetIndices to %d tags."), NetworkGameplayTagNodeIndex.Num() );
 
+	NetworkGameplayTagNodeIndexHash = 0;
+
 	for (FGameplayTagNetIndex i = 0; i < NetworkGameplayTagNodeIndex.Num(); i++)
 	{
 		if (NetworkGameplayTagNodeIndex[i].IsValid())
 		{
 			NetworkGameplayTagNodeIndex[i]->NetIndex = i;
 
+			NetworkGameplayTagNodeIndexHash = FCrc::StrCrc32(*NetworkGameplayTagNodeIndex[i]->GetCompleteTagString(), NetworkGameplayTagNodeIndexHash);
+
 			UE_CLOG(PrintNetIndiceAssignment, LogGameplayTags, Display, TEXT("Assigning NetIndex (%d) to Tag (%s)"), i, *NetworkGameplayTagNodeIndex[i]->GetCompleteTag().ToString());
 		}
 		else
@@ -491,6 +496,8 @@ void UGameplayTagsManager::ConstructNetIndex()
 			UE_LOG(LogGameplayTags, Warning, TEXT("TagNode Indice %d is invalid!"), i);
 		}
 	}
+
+	UE_LOG(LogGameplayTags, Log, TEXT("NetworkGameplayTagNodeIndexHash is %x"), NetworkGameplayTagNodeIndexHash);
 }
 
 FName UGameplayTagsManager::GetTagNameFromNetIndex(FGameplayTagNetIndex Index) const
