Fix tag hash being vulnerable to FName casing issues between client and server

diff --git a/Engine/Source/Runtime/GameplayTags/Private/GameplayTagsManager.cpp b/Engine/Source/Runtime/GameplayTags/Private/GameplayTagsManager.cpp
--- a/Engine/Source/Runtime/GameplayTags/Private/GameplayTagsManager.cpp
+++ b/Engine/Source/Runtime/GameplayTags/Private/GameplayTagsManager.cpp
@@ -487,7 +487,7 @@ void UGameplayTagsManager::ConstructNetIndex()
 		{
 			NetworkGameplayTagNodeIndex[i]->NetIndex = i;
 
-			NetworkGameplayTagNodeIndexHash = FCrc::StrCrc32(*NetworkGameplayTagNodeIndex[i]->GetCompleteTagString(), NetworkGameplayTagNodeIndexHash);
+			NetworkGameplayTagNodeIndexHash = FCrc::StrCrc32(*NetworkGameplayTagNodeIndex[i]->GetCompleteTagString().ToLower(), NetworkGameplayTagNodeIndexHash);
 
 			UE_CLOG(PrintNetIndiceAssignment, LogGameplayTags, Display, TEXT("Assigning NetIndex (%d) to Tag (%s)"), i, *NetworkGameplayTagNodeIndex[i]->GetCompleteTag().ToString());
 		}
