Fix static analysis warning about using WITH_EDITOR in an if block

diff --git a/Engine/Source/Runtime/GameplayTags/Private/GameplayTagsManager.cpp b/Engine/Source/Runtime/GameplayTags/Private/GameplayTagsManager.cpp
--- a/Engine/Source/Runtime/GameplayTags/Private/GameplayTagsManager.cpp
+++ b/Engine/Source/Runtime/GameplayTags/Private/GameplayTagsManager.cpp
@@ -60,8 +60,9 @@ void UGameplayTagsManager::LoadGameplayTagTables(bool bAllowAsyncLoad)
 	UGameplayTagsSettings* MutableDefault = GetMutableDefault<UGameplayTagsSettings>();
 	GameplayTagTables.Empty();
 
+#if !WITH_EDITOR
 	// If we're a cooked build and in a safe spot, start an async load so we can pipeline it
-	if (bAllowAsyncLoad && !WITH_EDITOR && !IsLoading() && MutableDefault->GameplayTagTableList.Num() > 0)
+	if (bAllowAsyncLoad && !IsLoading() && MutableDefault->GameplayTagTableList.Num() > 0)
 	{
 		for (FSoftObjectPath DataTablePath : MutableDefault->GameplayTagTableList)
 		{
@@ -70,6 +71,7 @@ void UGameplayTagsManager::LoadGameplayTagTables(bool bAllowAsyncLoad)
 
 		return;
 	}
+#endif // !WITH_EDITOR
 
 	SCOPE_LOG_GAMEPLAYTAGS(TEXT("UGameplayTagsManager::LoadGameplayTagTables"));
 	for (FSoftObjectPath DataTablePath : MutableDefault->GameplayTagTableList)
