#UE4 Improve warning on invalid gameplay tags.

diff --git a/Engine/Source/Runtime/GameplayTags/Private/GameplayTagsManager.cpp b/Engine/Source/Runtime/GameplayTags/Private/GameplayTagsManager.cpp
--- a/Engine/Source/Runtime/GameplayTags/Private/GameplayTagsManager.cpp
+++ b/Engine/Source/Runtime/GameplayTags/Private/GameplayTagsManager.cpp
@@ -12,6 +12,7 @@
 #include "UObject/UObjectIterator.h"
 #include "UObject/LinkerLoad.h"
 #include "UObject/Package.h"
+#include "UObject/UObjectThreadContext.h"
 #include "GameplayTagsSettings.h"
 #include "GameplayTagsModule.h"
 #include "Framework/Notifications/NotificationManager.h"
@@ -636,7 +637,9 @@ void UGameplayTagsManager::RedirectTagsForContainer(FGameplayTagContainer& Conta
 			FGameplayTag OldTag = RequestGameplayTag(TagName, false);
 			if (!OldTag.IsValid() && ShouldWarnOnInvalidTags())
 			{
-				UE_LOG(LogGameplayTags, Warning, TEXT("Invalid GameplayTag %s found while loading property %s."), *TagName.ToString(), *GetPathNameSafe(SerializingProperty));
+				FUObjectSerializeContext* LoadContext = FUObjectThreadContext::Get().GetSerializeContext();
+				UObject* LoadingObject = LoadContext ? LoadContext->SerializedObject : nullptr;
+				UE_LOG(LogGameplayTags, Warning, TEXT("Invalid GameplayTag %s found while loading %s in property %s."), *TagName.ToString(), *GetPathNameSafe(LoadingObject), *GetPathNameSafe(SerializingProperty));
 			}
 		}
 #endif
@@ -682,7 +685,9 @@ void UGameplayTagsManager::RedirectSingleGameplayTag(FGameplayTag& Tag, UPropert
 		FGameplayTag OldTag = RequestGameplayTag(TagName, false);
 		if (!OldTag.IsValid() && ShouldWarnOnInvalidTags())
 		{
-			UE_LOG(LogGameplayTags, Warning, TEXT("Invalid GameplayTag %s found while loading property %s."), *TagName.ToString(), *GetPathNameSafe(SerializingProperty));
+			FUObjectSerializeContext* LoadContext = FUObjectThreadContext::Get().GetSerializeContext();
+			UObject* LoadingObject = LoadContext ? LoadContext->SerializedObject : nullptr;
+			UE_LOG(LogGameplayTags, Warning, TEXT("Invalid GameplayTag %s found while loading %s in property %s."), *TagName.ToString(), *GetPathNameSafe(LoadingObject), *GetPathNameSafe(SerializingProperty));
 		}
 	}
 #endif
