Added a function to allow Blueprint to check if a GameplayTagQuery is empty. Allows queries being used as filters in Blueprint to have an unset / off state.

diff --git a/Engine/Source/Runtime/GameplayTags/Classes/BlueprintGameplayTagLibrary.h b/Engine/Source/Runtime/GameplayTags/Classes/BlueprintGameplayTagLibrary.h
--- a/Engine/Source/Runtime/GameplayTags/Classes/BlueprintGameplayTagLibrary.h
+++ b/Engine/Source/Runtime/GameplayTags/Classes/BlueprintGameplayTagLibrary.h
@@ -106,6 +106,16 @@ class UBlueprintGameplayTagLibrary : public UBlueprintFunctionLibrary
 	UFUNCTION(BlueprintPure, Category="GameplayTags", meta = (Keywords = "DoesContainerMatchAllTagsInContainer", BlueprintThreadSafe))
 	static bool HasAllTags(const FGameplayTagContainer& TagContainer, const FGameplayTagContainer& OtherContainer, bool bExactMatch);
 
+	/**
+	 * Check if the specified tag query is empty
+	 * 
+	 * @param TagQuery				Query to check
+	 * 
+	 * @return True if the query is empty, false otherwise.
+	 */
+	UFUNCTION(BlueprintPure, Category = "GameplayTags", meta = (BlueprintThreadSafe))
+	static bool IsTagQueryEmpty(const FGameplayTagQuery& TagQuery);
+
 	/**
 	 * Check if the specified tag container matches the given Tag Query
 	 * 
diff --git a/Engine/Source/Runtime/GameplayTags/Private/BlueprintGameplayTagLibrary.cpp b/Engine/Source/Runtime/GameplayTags/Private/BlueprintGameplayTagLibrary.cpp
--- a/Engine/Source/Runtime/GameplayTags/Private/BlueprintGameplayTagLibrary.cpp
+++ b/Engine/Source/Runtime/GameplayTags/Private/BlueprintGameplayTagLibrary.cpp
@@ -86,6 +86,11 @@ bool UBlueprintGameplayTagLibrary::HasAllTags(const FGameplayTagContainer& TagCo
 	return TagContainer.HasAll(OtherContainer);
 }
 
+bool UBlueprintGameplayTagLibrary::IsTagQueryEmpty(const FGameplayTagQuery& TagQuery)
+{
+	return TagQuery.IsEmpty();
+}
+
 bool UBlueprintGameplayTagLibrary::DoesContainerMatchTagQuery(const FGameplayTagContainer& TagContainer, const FGameplayTagQuery& TagQuery)
 {
 	return TagQuery.Matches(TagContainer);
