UGameplayTagsManager::AddTagIniSearchPath from
#RXN

diff --git a/Engine/Source/Runtime/GameplayTags/Classes/GameplayTagsManager.h b/Engine/Source/Runtime/GameplayTags/Classes/GameplayTagsManager.h
--- a/Engine/Source/Runtime/GameplayTags/Classes/GameplayTagsManager.h
+++ b/Engine/Source/Runtime/GameplayTags/Classes/GameplayTagsManager.h
@@ -425,6 +425,9 @@ class GAMEPLAYTAGS_API UGameplayTagsManager : public UObject
 	/** Loads the tag tables referenced in the GameplayTagSettings object */
 	void LoadGameplayTagTables(bool bAllowAsyncLoad = false);
 
+	/** Loads tag inis contained in the specified path */
+	void AddTagIniSearchPath(const FString& RootDir);
+
 	/** Helper function to construct the gameplay tag tree */
 	void ConstructGameplayTagTree();
 
@@ -658,6 +661,8 @@ private:
 
 	void AddChildrenTags(FGameplayTagContainer& TagContainer, TSharedPtr<FGameplayTagNode> GameplayTagNode, bool RecurseAll=true, bool OnlyIncludeDictionaryTags=false) const;
 
+	void AddTagsFromAdditionalLooseIniFiles(const TArray<FString>& IniFileList);
+
 	/**
 	 * Helper function for GameplayTagsMatch to get all parents when doing a parent match,
 	 * NOTE: Must never be made public as it uses the FNames which should never be exposed
@@ -692,6 +697,12 @@ private:
 	/** List of native tags to add when reconstructing tree */
 	TSet<FName> NativeTagsToAdd;
 
+	TSet<FName> RestrictedGameplayTagSourceNames;
+
+	TArray<FString> ExtraTagIniList;
+
+	bool bIsConstructingGameplayTagTree = false;
+
 	/** Cached runtime value for whether we are using fast replication or not. Initialized from config setting. */
 	bool bUseFastReplication;
 
diff --git a/Engine/Source/Runtime/GameplayTags/Private/GameplayTagsManager.cpp b/Engine/Source/Runtime/GameplayTags/Private/GameplayTagsManager.cpp
--- a/Engine/Source/Runtime/GameplayTags/Private/GameplayTagsManager.cpp
+++ b/Engine/Source/Runtime/GameplayTags/Private/GameplayTagsManager.cpp
@@ -105,15 +105,88 @@ struct FCompareFGameplayTagNodeByTag
 	}
 };
 
+void UGameplayTagsManager::AddTagIniSearchPath(const FString& RootDir)
+{
+	// Read all tags from the ini
+	TArray<FString> FilesInDirectory;
+	IFileManager::Get().FindFilesRecursive(FilesInDirectory, *RootDir, TEXT("*.ini"), true, false);
+	FilesInDirectory.Sort();
+	for (const FString& IniFilePath : FilesInDirectory)
+	{
+		ExtraTagIniList.AddUnique(IniFilePath);
+	}
+
+	if (!bIsConstructingGameplayTagTree)
+	{
+		AddTagsFromAdditionalLooseIniFiles(FilesInDirectory);
+
+		IGameplayTagsModule::OnGameplayTagTreeChanged.Broadcast();
+#if WITH_EDITOR
+		EditorRefreshGameplayTagTree();
+#endif
+	}
+}
+
+void UGameplayTagsManager::AddTagsFromAdditionalLooseIniFiles(const TArray<FString>& IniFileList)
+{
+	// Read all tags from the ini
+	for (const FString& IniFilePath : IniFileList)
+	{
+		const FName TagSource = FName(*FPaths::GetCleanFilename(IniFilePath));
+
+		// skip the restricted tag files
+		if (RestrictedGameplayTagSourceNames.Contains(TagSource))
+		{
+			continue;
+		}
+
+		FGameplayTagSource* FoundSource = FindOrAddTagSource(TagSource, EGameplayTagSourceType::TagList);
+
+		UE_CLOG(GAMEPLAYTAGS_VERBOSE, LogGameplayTags, Display, TEXT("Loading Tag File: %s"), *IniFilePath);
+
+		if (FoundSource && FoundSource->SourceTagList)
+		{
+			FoundSource->SourceTagList->ConfigFileName = IniFilePath;
+
+			// Check deprecated locations
+			TArray<FString> Tags;
+			if (GConfig->GetArray(TEXT("UserTags"), TEXT("GameplayTags"), Tags, IniFilePath))
+			{
+				for (const FString& Tag : Tags)
+				{
+					FoundSource->SourceTagList->GameplayTagList.AddUnique(FGameplayTagTableRow(FName(*Tag)));
+				}
+			}
+			else
+			{
+				// Load from new ini
+				FoundSource->SourceTagList->LoadConfig(UGameplayTagsList::StaticClass(), *IniFilePath);
+			}
+
+#if WITH_EDITOR
+			if (GIsEditor || IsRunningCommandlet()) // Sort tags for UI Purposes but don't sort in -game scenario since this would break compat with noneditor cooked builds
+			{
+				FoundSource->SourceTagList->SortTags();
+			}
+#endif
+
+			for (const FGameplayTagTableRow& TableRow : FoundSource->SourceTagList->GameplayTagList)
+			{
+				AddTagTableRow(TableRow, TagSource);
+			}
+		}
+	}
+}
+
 void UGameplayTagsManager::ConstructGameplayTagTree()
 {
 	SCOPE_LOG_GAMEPLAYTAGS(TEXT("UGameplayTagsManager::ConstructGameplayTagTree"));
+	TGuardValue<bool> GuardRebuilding(bIsConstructingGameplayTagTree, true);
 	if (!GameplayRootTag.IsValid())
 	{
 		GameplayRootTag = MakeShareable(new FGameplayTagNode());
 
 		UGameplayTagsSettings* MutableDefault = GetMutableDefault<UGameplayTagsSettings>();
-		TArray<FName> RestrictedGameplayTagSourceNames;
 
 		// Copy invalid characters, then add internal ones
 		InvalidTagCharacters = MutableDefault->InvalidTagCharacters;
@@ -223,65 +296,8 @@ void UGameplayTagsManager::ConstructGameplayTagTree()
 			}
 
 			// Extra tags
-		
-			// Read all tags from the ini
-			TArray<FString> FilesInDirectory;
-			IFileManager::Get().FindFilesRecursive(FilesInDirectory, *(FPaths::ProjectConfigDir() / TEXT("Tags")), TEXT("*.ini"), true, false);
-			FilesInDirectory.Sort();
-			for (FString& FileName : FilesInDirectory)
-			{
-				TagSource = FName(*FPaths::GetCleanFilename(FileName));
-
-				// skip the restricted tag files
-				bool bIsRestrictedTagFile = false;
-				for (const FName& RestrictedTagSource : RestrictedGameplayTagSourceNames)
-				{
-					if (TagSource == RestrictedTagSource)
-					{
-						bIsRestrictedTagFile = true;
-						break;
-					}
-				}
-
-				if (bIsRestrictedTagFile)
-				{
-					continue;
-				}
-
-				FGameplayTagSource* FoundSource = FindOrAddTagSource(TagSource, EGameplayTagSourceType::TagList);
-
-				UE_CLOG(GAMEPLAYTAGS_VERBOSE, LogGameplayTags, Display, TEXT("Loading Tag File: %s"), *FileName);
-
-				if (FoundSource && FoundSource->SourceTagList)
-				{
-					// Check deprecated locations
-					TArray<FString> Tags;
-					if (GConfig->GetArray(TEXT("UserTags"), TEXT("GameplayTags"), Tags, FileName))
-					{
-						for (const FString& Tag : Tags)
-						{
-							FoundSource->SourceTagList->GameplayTagList.AddUnique(FGameplayTagTableRow(FName(*Tag)));
-						}
-					}
-					else
-					{
-						// Load from new ini
-						FoundSource->SourceTagList->LoadConfig(UGameplayTagsList::StaticClass(), *FileName);
-					}
-
-#if WITH_EDITOR
-					if (GIsEditor || IsRunningCommandlet()) // Sort tags for UI Purposes but don't sort in -game scenario since this would break compat with noneditor cooked builds
-					{
-						FoundSource->SourceTagList->SortTags();
-					}
-#endif
-
-					for (const FGameplayTagTableRow& TableRow : FoundSource->SourceTagList->GameplayTagList)
-					{
-						AddTagTableRow(TableRow, TagSource);
-					}
-				}
-			}
+			AddTagIniSearchPath(FPaths::ProjectConfigDir() / TEXT("Tags"));
+			AddTagsFromAdditionalLooseIniFiles(ExtraTagIniList);
 		}
 
 #if WITH_EDITOR
@@ -929,6 +945,7 @@ void UGameplayTagsManager::DestroyGameplayTagTree()
 		GameplayRootTag.Reset();
 		GameplayTagNodeMap.Reset();
 	}
+	RestrictedGameplayTagSourceNames.Reset();
 }
 
 bool UGameplayTagsManager::IsNativelyAddedTag(FGameplayTag Tag) const
@@ -1372,7 +1389,6 @@ void UGameplayTagsManager::EditorRefreshGameplayTagTree()
 	DestroyGameplayTagTree();
 	LoadGameplayTagTables(false);
 	ConstructGameplayTagTree();
-
 	OnEditorRefreshGameplayTagTree.Broadcast();
 }
 
