Gameplay Tags: Added AddTagIniSearchPath to the tag manager, which allows users to specify additional search paths for gameplay tag ini files in addition to the default GameDir\Config\Tags\*.ini directory
[remainder of change previously checked in as 7248919, shelf wasn't fully up to date]

diff --git a/Engine/Source/Runtime/GameplayTags/Private/GameplayTagsManager.cpp b/Engine/Source/Runtime/GameplayTags/Private/GameplayTagsManager.cpp
--- a/Engine/Source/Runtime/GameplayTags/Private/GameplayTagsManager.cpp
+++ b/Engine/Source/Runtime/GameplayTags/Private/GameplayTagsManager.cpp
@@ -118,11 +118,14 @@ void UGameplayTagsManager::AddTagIniSearchPath(const FString& RootDir)
 
 	if (!bIsConstructingGameplayTagTree)
 	{
+#if WITH_EDITOR
+		EditorRefreshGameplayTagTree();
+#else
 		AddTagsFromAdditionalLooseIniFiles(FilesInDirectory);
 
+		ConstructNetIndex();
+
 		IGameplayTagsModule::OnGameplayTagTreeChanged.Broadcast();
-#if WITH_EDITOR
-		EditorRefreshGameplayTagTree();
 #endif
 	}
 }
