Improve gameplay tag table load times by scheduling an async load for all of them and then flushing, instead of doing individual sync loads
Clean up the tag manager load logging so there's only one log statement by default, and can opt into more
Add proper error validation for gameplay tags, the old parsing code would silently fixup tags ending in . by accident
The list of invalid characters defaults to "', but can be changed in the tag settings
Add logging for tag editor error notifications
Copy of CL #4662909 and #4666869

diff --git a/Engine/Plugins/Editor/GameplayTagsEditor/Source/GameplayTagsEditor/Private/GameplayTagsEditorModule.cpp b/Engine/Plugins/Editor/GameplayTagsEditor/Source/GameplayTagsEditor/Private/GameplayTagsEditorModule.cpp
--- a/Engine/Plugins/Editor/GameplayTagsEditor/Source/GameplayTagsEditor/Private/GameplayTagsEditorModule.cpp
+++ b/Engine/Plugins/Editor/GameplayTagsEditor/Source/GameplayTagsEditor/Private/GameplayTagsEditorModule.cpp
@@ -177,12 +177,18 @@ public:
 		return true;
 	}
 
-	void ShowNotification(const FText& TextToDisplay, float TimeToDisplay)
+	void ShowNotification(const FText& TextToDisplay, float TimeToDisplay, bool bLogError = false)
 	{
 		FNotificationInfo Info(TextToDisplay);
 		Info.ExpireDuration = TimeToDisplay;
 
 		FSlateNotificationManager::Get().AddNotification(Info);
+
+		// Also log if error
+		if (bLogError)
+		{
+			UE_LOG(LogGameplayTags, Error, TEXT("%s"), *TextToDisplay.ToString())
+		}
 	}
 
 	void MigrateSettings()
@@ -325,6 +331,14 @@ public:
 		UGameplayTagsSettings*			Settings = GetMutableDefault<UGameplayTagsSettings>();
 		UGameplayTagsDeveloperSettings* DevSettings = GetMutableDefault<UGameplayTagsDeveloperSettings>();
 
+		FText ErrorText;
+		FString FixedString;
+		if (!Manager.IsValidGameplayTagString(NewTag, &ErrorText, &FixedString))
+		{
+			ShowNotification(FText::Format(LOCTEXT("AddTagFailure_BadString", "Failed to add gameplay tag {0}: {1}, try {2} instead!"), FText::FromString(NewTag), ErrorText, FText::FromString(FixedString)), 10.0f, true);
+			return false;
+		}
+
 		FName NewTagName = FName(*NewTag);
 
 		// Delete existing redirector
@@ -333,7 +347,7 @@ public:
 		// Already in the list as an explicit tag, ignore. Note we want to add if it is in implicit tag. (E.g, someone added A.B.C then someone tries to add A.B)
 		if (Manager.IsDictionaryTag(NewTagName))
 		{
-			ShowNotification(FText::Format(LOCTEXT("AddTagFailure_AlreadyExists", "Failed to add gameplay tag {0}, already exists!"), FText::FromString(NewTag)), 10.0f);
+			ShowNotification(FText::Format(LOCTEXT("AddTagFailure_AlreadyExists", "Failed to add gameplay tag {0}, already exists!"), FText::FromString(NewTag)), 10.0f, true);
 
 			return false;
 		}
@@ -358,7 +372,7 @@ public:
 					{
 						break;
 					}
-					ShowNotification(FText::Format(LOCTEXT("AddRestrictedTagFailure", "Failed to add restricted gameplay tag {0}, {1} is not a restricted tag"), FText::FromString(NewTag), FText::FromString(AncestorTag)), 10.0f);
+					ShowNotification(FText::Format(LOCTEXT("AddRestrictedTagFailure", "Failed to add restricted gameplay tag {0}, {1} is not a restricted tag"), FText::FromString(NewTag), FText::FromString(AncestorTag)), 10.0f, true);
 
 					return false;
 				}
@@ -389,7 +403,7 @@ public:
 							break;
 						}
 
-						ShowNotification(FText::Format(LOCTEXT("AddTagFailure_RestrictedTag", "Failed to add gameplay tag {0}, {1} is a restricted tag and does not allow non-restricted children"), FText::FromString(NewTag), FText::FromString(AncestorTag)), 10.0f);
+						ShowNotification(FText::Format(LOCTEXT("AddTagFailure_RestrictedTag", "Failed to add gameplay tag {0}, {1} is a restricted tag and does not allow non-restricted children"), FText::FromString(NewTag), FText::FromString(AncestorTag)), 10.0f, true);
 
 						return false;
 					}
@@ -455,7 +469,7 @@ public:
 		
 		if (!bSuccess)
 		{
-			ShowNotification(FText::Format(LOCTEXT("AddTagFailure", "Failed to add gameplay tag {0} to dictionary {1}!"), FText::FromString(NewTag), FText::FromName(TagSourceName)), 10.0f);
+			ShowNotification(FText::Format(LOCTEXT("AddTagFailure", "Failed to add gameplay tag {0} to dictionary {1}!"), FText::FromString(NewTag), FText::FromName(TagSourceName)), 10.0f, true);
 
 			return false;
 		}
@@ -490,7 +504,7 @@ public:
 		
 		if (!Manager.GetTagEditorData(TagName, Comment, TagSourceName, bTagIsExplicit, bTagIsRestricted, bTagAllowsNonRestrictedChildren))
 		{
-			ShowNotification(FText::Format(LOCTEXT("RemoveTagFailureNoTag", "Cannot delete tag {0}, does not exist!"), FText::FromName(TagName)), 10.0f);
+			ShowNotification(FText::Format(LOCTEXT("RemoveTagFailureNoTag", "Cannot delete tag {0}, does not exist!"), FText::FromName(TagName)), 10.0f, true);
 			return false;
 		}
 
@@ -501,19 +515,19 @@ public:
 		// Check if the tag is implicitly defined
 		if (!bTagIsExplicit || !TagSource)
 		{
-			ShowNotification(FText::Format(LOCTEXT("RemoveTagFailureNoSource", "Cannot delete tag {0} as it is implicit, remove children manually"), FText::FromName(TagName)), 10.0f);
+			ShowNotification(FText::Format(LOCTEXT("RemoveTagFailureNoSource", "Cannot delete tag {0} as it is implicit, remove children manually"), FText::FromName(TagName)), 10.0f, true);
 			return false;
 		}
 		
 		if (bTagIsRestricted && !TagSource->SourceRestrictedTagList)
 		{
-			ShowNotification(FText::Format(LOCTEXT("RemoveTagFailureBadSource", "Cannot delete tag {0} from source {1}, remove manually"), FText::FromName(TagName), FText::FromName(TagSourceName)), 10.0f);
+			ShowNotification(FText::Format(LOCTEXT("RemoveTagFailureBadSource", "Cannot delete tag {0} from source {1}, remove manually"), FText::FromName(TagName), FText::FromName(TagSourceName)), 10.0f, true);
 			return false;
 		}
 
 		if (!bTagIsRestricted && !TagSource->SourceTagList)
 		{
-			ShowNotification(FText::Format(LOCTEXT("RemoveTagFailureBadSource", "Cannot delete tag {0} from source {1}, remove manually"), FText::FromName(TagName), FText::FromName(TagSourceName)), 10.0f);
+			ShowNotification(FText::Format(LOCTEXT("RemoveTagFailureBadSource", "Cannot delete tag {0} from source {1}, remove manually"), FText::FromName(TagName), FText::FromName(TagSourceName)), 10.0f, true);
 			return false;
 		}
 
@@ -554,7 +568,7 @@ public:
 
 			if (Referencers.Num() > 0)
 			{
-				ShowNotification(FText::Format(LOCTEXT("RemoveTagFailureBadSource_Referenced", "Cannot delete tag {0}, still referenced by {1} and possibly others"), FText::FromName(TagNameToDelete), FText::FromString(Referencers[0].ToString())), 10.0f);
+				ShowNotification(FText::Format(LOCTEXT("RemoveTagFailureBadSource_Referenced", "Cannot delete tag {0}, still referenced by {1} and possibly others"), FText::FromName(TagNameToDelete), FText::FromString(Referencers[0].ToString())), 10.0f, true);
 
 				return false;
 			}
@@ -609,7 +623,7 @@ public:
 			}
 		}
 
-		ShowNotification(FText::Format(LOCTEXT("RemoveTagFailureNoTag", "Cannot delete tag {0}, does not exist!"), FText::FromName(TagName)), 10.0f);
+		ShowNotification(FText::Format(LOCTEXT("RemoveTagFailureNoTag", "Cannot delete tag {0}, does not exist!"), FText::FromName(TagName)), 10.0f, true);
 		
 		return false;
 	}
@@ -640,7 +654,7 @@ public:
 					FGameplayTagContainer ChildTags = Manager.RequestGameplayTagDirectDescendantsInDictionary(ActualTag, EGameplayTagSelectionType::NonRestrictedOnly);
 					if (!ChildTags.IsEmpty())
 					{
-						ShowNotification(LOCTEXT("ToggleAllowNonRestrictedChildrenFailure", "Cannot prevent non-restricted children since some already exist! Delete them first."), 10.0f);
+						ShowNotification(LOCTEXT("ToggleAllowNonRestrictedChildrenFailure", "Cannot prevent non-restricted children since some already exist! Delete them first."), 10.0f, true);
 						return false;
 					}
 				}
@@ -734,7 +748,7 @@ public:
 			}
 			else
 			{
-				ShowNotification(FText::Format(LOCTEXT("RenameFailure", "Tag {0} redirector was created but original tag was not destroyed as it has children"), FText::FromString(TagToRename)), 10.0f);
+				ShowNotification(FText::Format(LOCTEXT("RenameFailure", "Tag {0} redirector was created but original tag was not destroyed as it has children"), FText::FromString(TagToRename)), 10.0f, true);
 			}
 		}
 
diff --git a/Engine/Source/Runtime/GameplayTags/Classes/GameplayTagContainer.h b/Engine/Source/Runtime/GameplayTags/Classes/GameplayTagContainer.h
--- a/Engine/Source/Runtime/GameplayTags/Classes/GameplayTagContainer.h
+++ b/Engine/Source/Runtime/GameplayTags/Classes/GameplayTagContainer.h
@@ -59,13 +59,20 @@ struct GAMEPLAYTAGS_API FGameplayTag
 	 * Gets the FGameplayTag that corresponds to the TagName
 	 *
 	 * @param TagName The Name of the tag to search for
-	 * 
 	 * @param ErrorIfNotfound: ensure() that tag exists.
-	 * 
 	 * @return Will return the corresponding FGameplayTag or an empty one if not found.
 	 */
 	static FGameplayTag RequestGameplayTag(FName TagName, bool ErrorIfNotFound=true);
 
+	/** 
+	 * Returns true if this is a valid gameplay tag string (foo.bar.baz). If false, it will fill 
+	 * @param TagString String to check for validity
+	 * @param OutError If non-null and string invalid, will fill in with an error message
+	 * @param OutFixedString If non-null and string invalid, will attempt to fix. Will be empty if no fix is possible
+	 * @return True if this can be added to the tag dictionary, false if there's a syntax error
+	 */
+	static bool IsValidGameplayTagString(const FString& TagString, FText* OutError = nullptr, FString* OutFixedString = nullptr);
+
 	/** Operators */
 	FORCEINLINE bool operator==(FGameplayTag const& Other) const
 	{
diff --git a/Engine/Source/Runtime/GameplayTags/Classes/GameplayTagsManager.h b/Engine/Source/Runtime/GameplayTags/Classes/GameplayTagsManager.h
--- a/Engine/Source/Runtime/GameplayTags/Classes/GameplayTagsManager.h
+++ b/Engine/Source/Runtime/GameplayTags/Classes/GameplayTagsManager.h
@@ -13,6 +13,7 @@
 #include "GameplayTagsManager.generated.h"
 
 class UGameplayTagsList;
+struct FStreamableHandle;
 
 /** Simple struct for a table row in the gameplay tag table and element in the ini list */
 USTRUCT()
@@ -141,8 +142,8 @@ struct FGameplayTagNode
 	GENERATED_USTRUCT_BODY()
 	FGameplayTagNode(){};
 
-	/** Simple constructor */
-	FGameplayTagNode(FName InTag, TSharedPtr<FGameplayTagNode> InParentNode, bool InIsExplicitTag, bool InIsRestrictedTag, bool InAllowNonRestrictedChildren);
+	/** Simple constructor, passing redundant data for performance */
+	FGameplayTagNode(FName InTag, FName InFullTag, TSharedPtr<FGameplayTagNode> InParentNode, bool InIsExplicitTag, bool InIsRestrictedTag, bool InAllowNonRestrictedChildren);
 
 	/** Returns a correctly constructed container with only this tag, useful for doing container queries */
 	FORCEINLINE const FGameplayTagContainer& GetSingleTagContainer() const { return CompleteTagWithParents; }
@@ -304,6 +305,15 @@ class GAMEPLAYTAGS_API UGameplayTagsManager : public UObject
 	 */
 	FGameplayTag RequestGameplayTag(FName TagName, bool ErrorIfNotFound=true) const;
 
+	/** 
+	 * Returns true if this is a valid gameplay tag string (foo.bar.baz). If false, it will fill 
+	 * @param TagString String to check for validity
+	 * @param OutError If non-null and string invalid, will fill in with an error message
+	 * @param OutFixedString If non-null and string invalid, will attempt to fix. Will be empty if no fix is possible
+	 * @return True if this can be added to the tag dictionary, false if there's a syntax error
+	 */
+	bool IsValidGameplayTagString(const FString& TagString, FText* OutError = nullptr, FString* OutFixedString = nullptr);
+
 	/**
 	 *	Searches for a gameplay tag given a partial string. This is slow and intended mainly for console commands/utilities to make
 	 *	developer life's easier. This will attempt to match as best as it can. If you pass "A.b" it will match on "A.b." before it matches "a.b.c".
@@ -413,7 +423,7 @@ class GAMEPLAYTAGS_API UGameplayTagsManager : public UObject
 	}
 
 	/** Loads the tag tables referenced in the GameplayTagSettings object */
-	void LoadGameplayTagTables();
+	void LoadGameplayTagTables(bool bAllowAsyncLoad = false);
 
 	/** Helper function to construct the gameplay tag tree */
 	void ConstructGameplayTagTree();
@@ -615,7 +625,8 @@ private:
 	/**
 	 * Helper function to insert a tag into a tag node array
 	 *
-	 * @param Tag							Tag to insert
+	 * @param Tag							Short name of tag to insert
+	 * @param FullTag						Full tag, passed in for performance
 	 * @param ParentNode					Parent node, if any, for the tag
 	 * @param NodeArray						Node array to insert the new node into, if necessary (if the tag already exists, no insertion will occur)
 	 * @param SourceName					File tag was added from
@@ -626,7 +637,7 @@ private:
 	 *
 	 * @return Index of the node of the tag
 	 */
-	int32 InsertTagIntoNodeArray(FName Tag, TSharedPtr<FGameplayTagNode> ParentNode, TArray< TSharedPtr<FGameplayTagNode> >& NodeArray, FName SourceName, const FString& DevComment, bool bIsExplicitTag, bool bIsRestrictedTag, bool bAllowNonRestrictedChildren);
+	int32 InsertTagIntoNodeArray(FName Tag, FName FullTag, TSharedPtr<FGameplayTagNode> ParentNode, TArray< TSharedPtr<FGameplayTagNode> >& NodeArray, FName SourceName, const FString& DevComment, bool bIsExplicitTag, bool bIsRestrictedTag, bool bAllowNonRestrictedChildren);
 
 	/** Helper function to populate the tag tree from each table */
 	void PopulateTreeFromDataTable(class UDataTable* Table);
@@ -678,6 +689,9 @@ private:
 	/** True if native tags have all been added and flushed */
 	bool bDoneAddingNativeTags;
 
+	/** String with outlawed characters inside tags */
+	FString InvalidTagCharacters;
+
 #if WITH_EDITOR
 	// This critical section is to handle an editor-only issue where tag requests come from another thread when async loading from a background thread in FGameplayTagContainer::Serialize.
 	// This class is not generically threadsafe.
@@ -690,9 +704,6 @@ private:
 	/** Sorted list of nodes, used for network replication */
 	TArray<TSharedPtr<FGameplayTagNode>> NetworkGameplayTagNodeIndex;
 
-	UPROPERTY()
-	TArray<UDataTable*> RestrictedGameplayTagTables;
-
 	/** Holds all of the valid gameplay-related tags that can be applied to assets */
 	UPROPERTY()
 	TArray<UDataTable*> GameplayTagTables;
diff --git a/Engine/Source/Runtime/GameplayTags/Classes/GameplayTagsSettings.h b/Engine/Source/Runtime/GameplayTags/Classes/GameplayTagsSettings.h
--- a/Engine/Source/Runtime/GameplayTags/Classes/GameplayTagsSettings.h
+++ b/Engine/Source/Runtime/GameplayTags/Classes/GameplayTagsSettings.h
@@ -126,17 +126,22 @@ class GAMEPLAYTAGS_API UGameplayTagsSettings : public UGameplayTagsList
 	UPROPERTY(config, EditAnywhere, Category = GameplayTags)
 	bool ImportTagsFromConfig;
 
-	/** If true, will give load warnings when reading invalid tags off disk */
+	/** If true, will give load warnings when reading in saved tag references that are not in the dictionary */
 	UPROPERTY(config, EditAnywhere, Category = GameplayTags)
 	bool WarnOnInvalidTags;
 
-	UPROPERTY(config, EditAnywhere, Category = GameplayTags)
-	TArray<FGameplayTagCategoryRemap>	CategoryRemapping;
-
 	/** If true, will replicate gameplay tags by index instead of name. For this to work, tags must be identical on client and server */
 	UPROPERTY(config, EditAnywhere, Category = "Advanced Replication")
 	bool FastReplication;
 
+	/** These characters cannot be used in gameplay tags, in addition to special ones like newline*/
+	UPROPERTY(config, EditAnywhere, Category = GameplayTags)
+	FString InvalidTagCharacters;
+
+	/** Category remapping. This allows base engine tag category meta data to remap to multiple project-specific categories. */
+	UPROPERTY(config, EditAnywhere, Category = GameplayTags)
+	TArray<FGameplayTagCategoryRemap> CategoryRemapping;
+
 	/** List of data tables to load tags from */
 	UPROPERTY(config, EditAnywhere, Category = GameplayTags, meta = (AllowedClasses = "DataTable"))
 	TArray<FSoftObjectPath> GameplayTagTableList;
@@ -165,6 +170,7 @@ class GAMEPLAYTAGS_API UGameplayTagsSettings : public UGameplayTagsList
 	UPROPERTY(EditAnywhere, AdvancedDisplay, transient, Category = "Advanced Gameplay Tags")
 	FString RestrictedTagList;
 #endif
+
 #if WITH_EDITOR
 	virtual void PreEditChange(UProperty* PropertyThatWillChange) override;
 	virtual void PostEditChangeProperty(FPropertyChangedEvent& PropertyChangedEvent) override;
diff --git a/Engine/Source/Runtime/GameplayTags/Private/GameplayTagContainer.cpp b/Engine/Source/Runtime/GameplayTags/Private/GameplayTagContainer.cpp
--- a/Engine/Source/Runtime/GameplayTags/Private/GameplayTagContainer.cpp
+++ b/Engine/Source/Runtime/GameplayTags/Private/GameplayTagContainer.cpp
@@ -1090,6 +1090,11 @@ FGameplayTag FGameplayTag::RequestGameplayTag(FName TagName, bool ErrorIfNotFoun
 	return UGameplayTagsManager::Get().RequestGameplayTag(TagName, ErrorIfNotFound);
 }
 
+bool FGameplayTag::IsValidGameplayTagString(const FString& TagString, FText* OutError, FString* OutFixedString)
+{
+	return UGameplayTagsManager::Get().IsValidGameplayTagString(TagString, OutError, OutFixedString);
+}
+
 FGameplayTagContainer FGameplayTag::GetGameplayTagParents() const
 {
 	return UGameplayTagsManager::Get().RequestGameplayTagParents(*this);
@@ -1350,7 +1355,6 @@ void FGameplayTag::FromExportString(const FString& ExportString)
 
 FGameplayTagNativeAdder::FGameplayTagNativeAdder()
 {
-	UE_LOG(LogGameplayTags, Display, TEXT("FGameplayTagNativeAdder::FGameplayTagNativeAdder"));
 	UGameplayTagsManager::OnLastChanceToAddNativeTags().AddRaw(this, &FGameplayTagNativeAdder::AddTags);
 }
 
diff --git a/Engine/Source/Runtime/GameplayTags/Private/GameplayTagsManager.cpp b/Engine/Source/Runtime/GameplayTags/Private/GameplayTagsManager.cpp
--- a/Engine/Source/Runtime/GameplayTags/Private/GameplayTagsManager.cpp
+++ b/Engine/Source/Runtime/GameplayTags/Private/GameplayTagsManager.cpp
@@ -1,6 +1,5 @@
 // Copyright 1998-2019 Epic Games, Inc. All Rights Reserved.
 
-
 #include "GameplayTagsManager.h"
 #include "Engine/Engine.h"
 #include "HAL/PlatformFilemanager.h"
@@ -46,12 +45,32 @@ UGameplayTagsManager::UGameplayTagsManager(const FObjectInitializer& ObjectIniti
 	NumBitsForContainerSize = 6;
 }
 
-void UGameplayTagsManager::LoadGameplayTagTables()
+// Enable to turn on detailed startup logging
+#define GAMEPLAYTAGS_VERBOSE 0
+
+#if STATS && GAMEPLAYTAGS_VERBOSE
+#define SCOPE_LOG_GAMEPLAYTAGS(Name) SCOPE_LOG_TIME_IN_SECONDS(Name, nullptr)
+#else
+#define SCOPE_LOG_GAMEPLAYTAGS(Name)
+#endif
+
+void UGameplayTagsManager::LoadGameplayTagTables(bool bAllowAsyncLoad)
 {
+	UGameplayTagsSettings* MutableDefault = GetMutableDefault<UGameplayTagsSettings>();
 	GameplayTagTables.Empty();
 
-	UGameplayTagsSettings* MutableDefault = GetMutableDefault<UGameplayTagsSettings>();
+	// If we're a cooked build and in a safe spot, start an async load so we can pipeline it
+	if (bAllowAsyncLoad && !WITH_EDITOR && !IsLoading() && MutableDefault->GameplayTagTableList.Num() > 0)
+	{
+		for (FSoftObjectPath DataTablePath : MutableDefault->GameplayTagTableList)
+		{
+			LoadPackageAsync(DataTablePath.GetLongPackageName());
+		}
 
+		return;
+	}
+
+	SCOPE_LOG_GAMEPLAYTAGS(TEXT("UGameplayTagsManager::LoadGameplayTagTables"));
 	for (FSoftObjectPath DataTablePath : MutableDefault->GameplayTagTableList)
 	{
 		UDataTable* TagTable = LoadObject<UDataTable>(nullptr, *DataTablePath.ToString(), nullptr, LOAD_None, nullptr);
@@ -84,21 +103,22 @@ struct FCompareFGameplayTagNodeByTag
 
 void UGameplayTagsManager::ConstructGameplayTagTree()
 {
+	SCOPE_LOG_GAMEPLAYTAGS(TEXT("UGameplayTagsManager::ConstructGameplayTagTree"));
 	if (!GameplayRootTag.IsValid())
 	{
 		GameplayRootTag = MakeShareable(new FGameplayTagNode());
 
 		UGameplayTagsSettings* MutableDefault = GetMutableDefault<UGameplayTagsSettings>();
-		FString DefaultEnginePath = FString::Printf(TEXT("%sDefaultEngine.ini"), *FPaths::SourceConfigDir());
 		TArray<FName> RestrictedGameplayTagSourceNames;
 
+		// Copy invalid characters, then add internal ones
+		InvalidTagCharacters = MutableDefault->InvalidTagCharacters;
+		InvalidTagCharacters.Append(TEXT("\r\n\t"));
+
 		// Add prefixes first
 		if (ShouldImportTagsFromINI())
 		{
-#if STATS
-			FString PerfMessage = FString::Printf(TEXT("UGameplayTagsManager::ConstructGameplayTagTree: ImportINI prefixes"));
-			SCOPE_LOG_TIME_IN_SECONDS(*PerfMessage, nullptr)
-#endif
+			SCOPE_LOG_GAMEPLAYTAGS(TEXT("UGameplayTagsManager::ConstructGameplayTagTree: ImportINI prefixes"));
 
 			TArray<FString> RestrictedGameplayTagFiles;
 			GetRestrictedTagConfigFiles(RestrictedGameplayTagFiles);
@@ -135,23 +155,28 @@ void UGameplayTagsManager::ConstructGameplayTagTree()
 			}
 		}
 
-		// Add native tags before other tags
-		for (FName TagToAdd : NativeTagsToAdd)
 		{
-			AddTagTableRow(FGameplayTagTableRow(TagToAdd), FGameplayTagSource::GetNativeName());
+			SCOPE_LOG_GAMEPLAYTAGS(TEXT("UGameplayTagsManager::ConstructGameplayTagTree: Add native tags"));
+			// Add native tags before other tags
+			for (FName TagToAdd : NativeTagsToAdd)
+			{
+				AddTagTableRow(FGameplayTagTableRow(TagToAdd), FGameplayTagSource::GetNativeName());
+			}
 		}
 
+		// If we didn't load any tables it might be async loading, so load again with a flush
+		if (GameplayTagTables.Num() == 0)
 		{
-#if STATS
-			FString PerfMessage = FString::Printf(TEXT("UGameplayTagsManager::ConstructGameplayTagTree: Construct from data asset"));
-			SCOPE_LOG_TIME_IN_SECONDS(*PerfMessage, nullptr)
-#endif
-		
-			for (auto It(GameplayTagTables.CreateIterator()); It; It++)
+			LoadGameplayTagTables(false);
+		}
+
+		{
+			SCOPE_LOG_GAMEPLAYTAGS(TEXT("UGameplayTagsManager::ConstructGameplayTagTree: Construct from data asset"));
+			for (UDataTable* DataTable : GameplayTagTables)
 			{
-				if (*It)
+				if (DataTable)
 				{
-					PopulateTreeFromDataTable(*It);
+					PopulateTreeFromDataTable(DataTable);
 				}
 			}
 		}
@@ -161,14 +186,12 @@ void UGameplayTagsManager::ConstructGameplayTagTree()
 
 		if (ShouldImportTagsFromINI())
 		{
-#if STATS
-			FString PerfMessage = FString::Printf(TEXT("UGameplayTagsManager::ConstructGameplayTagTree: ImportINI"));
-			SCOPE_LOG_TIME_IN_SECONDS(*PerfMessage, nullptr)
-#endif
+			SCOPE_LOG_GAMEPLAYTAGS(TEXT("UGameplayTagsManager::ConstructGameplayTagTree: ImportINI tags"));
+
 			// Copy from deprecated list in DefaultEngine.ini
 			TArray<FString> EngineConfigTags;
-			GConfig->GetArray(TEXT("/Script/GameplayTags.GameplayTagsSettings"), TEXT("+GameplayTags"), EngineConfigTags, DefaultEnginePath);
-			
+			GConfig->GetArray(TEXT("/Script/GameplayTags.GameplayTagsSettings"), TEXT("+GameplayTags"), EngineConfigTags, GEngineIni);
+
 			for (const FString& EngineConfigTag : EngineConfigTags)
 			{
 				MutableDefault->GameplayTagList.AddUnique(FGameplayTagTableRow(FName(*EngineConfigTag)));
@@ -223,7 +246,7 @@ void UGameplayTagsManager::ConstructGameplayTagTree()
 
 				FGameplayTagSource* FoundSource = FindOrAddTagSource(TagSource, EGameplayTagSourceType::TagList);
 
-				UE_LOG(LogGameplayTags, Display, TEXT("Loading Tag File: %s"), *FileName);
+				UE_CLOG(GAMEPLAYTAGS_VERBOSE, LogGameplayTags, Display, TEXT("Loading Tag File: %s"), *FileName);
 
 				if (FoundSource && FoundSource->SourceTagList)
 				{
@@ -264,144 +287,143 @@ void UGameplayTagsManager::ConstructGameplayTagTree()
 			AddTagTableRow(FGameplayTagTableRow(TransientTag), FGameplayTagSource::GetTransientEditorName());
 		}
 #endif
-
-		// Grab the commonly replicated tags
-		CommonlyReplicatedTags.Empty();
-		for (FName TagName : MutableDefault->CommonlyReplicatedTags)
 		{
-			FGameplayTag Tag = RequestGameplayTag(TagName);
-			if (Tag.IsValid())
+			SCOPE_LOG_GAMEPLAYTAGS(TEXT("UGameplayTagsManager::ConstructGameplayTagTree: Request common tags"));
+			// Grab the commonly replicated tags
+			CommonlyReplicatedTags.Empty();
+			for (FName TagName : MutableDefault->CommonlyReplicatedTags)
 			{
-				CommonlyReplicatedTags.Add(Tag);
-			}
-			else
-			{
-				UE_LOG(LogGameplayTags, Warning, TEXT("%s was found in the CommonlyReplicatedTags list but doesn't appear to be a valid tag!"), *TagName.ToString());
+				FGameplayTag Tag = RequestGameplayTag(TagName);
+				if (Tag.IsValid())
+				{
+					CommonlyReplicatedTags.Add(Tag);
+				}
+				else
+				{
+					UE_LOG(LogGameplayTags, Warning, TEXT("%s was found in the CommonlyReplicatedTags list but doesn't appear to be a valid tag!"), *TagName.ToString());
+				}
 			}
-		}
 
-		bUseFastReplication = MutableDefault->FastReplication;
-		bShouldWarnOnInvalidTags = MutableDefault->WarnOnInvalidTags;
-		NumBitsForContainerSize = MutableDefault->NumBitsForContainerSize;
-		NetIndexFirstBitSegment = MutableDefault->NetIndexFirstBitSegment;
+			bUseFastReplication = MutableDefault->FastReplication;
+			bShouldWarnOnInvalidTags = MutableDefault->WarnOnInvalidTags;
+			NumBitsForContainerSize = MutableDefault->NumBitsForContainerSize;
+			NetIndexFirstBitSegment = MutableDefault->NetIndexFirstBitSegment;
+		}
 
 		if (ShouldUseFastReplication())
 		{
-#if STATS
-			FString PerfMessage = FString::Printf(TEXT("UGameplayTagsManager::ConstructGameplayTagTree: Reconstruct NetIndex"));
-			SCOPE_LOG_TIME_IN_SECONDS(*PerfMessage, nullptr)
-#endif
+			SCOPE_LOG_GAMEPLAYTAGS(TEXT("UGameplayTagsManager::ConstructGameplayTagTree: Reconstruct NetIndex"));
 			ConstructNetIndex();
 		}
 
 		{
-#if STATS
-			FString PerfMessage = FString::Printf(TEXT("UGameplayTagsManager::ConstructGameplayTagTree: GameplayTagTreeChangedEvent.Broadcast"));
-			SCOPE_LOG_TIME_IN_SECONDS(*PerfMessage, nullptr)
-#endif
+			SCOPE_LOG_GAMEPLAYTAGS(TEXT("UGameplayTagsManager::ConstructGameplayTagTree: GameplayTagTreeChangedEvent.Broadcast"));
 			IGameplayTagsModule::OnGameplayTagTreeChanged.Broadcast();
 		}
 
-		// Update the TagRedirects map
-		TagRedirects.Empty();
+		{
+			SCOPE_LOG_GAMEPLAYTAGS(TEXT("UGameplayTagsManager::ConstructGameplayTagTree: Load redirects"));
+			// Update the TagRedirects map
+			TagRedirects.Empty();
 
-		// Check the deprecated location
-		bool bFoundDeprecated = false;
-		FConfigSection* PackageRedirects = GConfig->GetSectionPrivate(TEXT("/Script/Engine.Engine"), false, true, DefaultEnginePath);
+			// Check the deprecated location
+			bool bFoundDeprecated = false;
+			FConfigSection* PackageRedirects = GConfig->GetSectionPrivate(TEXT("/Script/Engine.Engine"), false, true, GEngineIni);
 
-		if (PackageRedirects)
-		{
-			for (FConfigSection::TIterator It(*PackageRedirects); It; ++It)
+			if (PackageRedirects)
 			{
-				if (It.Key() == TEXT("+GameplayTagRedirects"))
+				for (FConfigSection::TIterator It(*PackageRedirects); It; ++It)
 				{
-					FName OldTagName = NAME_None;
-					FName NewTagName;
-
-					if (FParse::Value(*It.Value().GetValue(), TEXT("OldTagName="), OldTagName))
+					if (It.Key() == TEXT("+GameplayTagRedirects"))
 					{
-						if (FParse::Value(*It.Value().GetValue(), TEXT("NewTagName="), NewTagName))
+						FName OldTagName = NAME_None;
+						FName NewTagName;
+
+						if (FParse::Value(*It.Value().GetValue(), TEXT("OldTagName="), OldTagName))
 						{
-							FGameplayTagRedirect Redirect;
-							Redirect.OldTagName = OldTagName;
-							Redirect.NewTagName = NewTagName;
+							if (FParse::Value(*It.Value().GetValue(), TEXT("NewTagName="), NewTagName))
+							{
+								FGameplayTagRedirect Redirect;
+								Redirect.OldTagName = OldTagName;
+								Redirect.NewTagName = NewTagName;
 
-							MutableDefault->GameplayTagRedirects.AddUnique(Redirect);
+								MutableDefault->GameplayTagRedirects.AddUnique(Redirect);
 
-							bFoundDeprecated = true;
+								bFoundDeprecated = true;
+							}
 						}
 					}
 				}
 			}
-		}
-
-		if (bFoundDeprecated)
-		{
-			UE_LOG(LogGameplayTags, Log, TEXT("GameplayTagRedirects is in a deprecated location, after editing GameplayTags developer settings you must remove these manually"));
-		}
 
-		// Check settings object
-		for (const FGameplayTagRedirect& Redirect : MutableDefault->GameplayTagRedirects)
-		{
-			FName OldTagName = Redirect.OldTagName;
-			FName NewTagName = Redirect.NewTagName;
+			if (bFoundDeprecated)
+			{
+				UE_LOG(LogGameplayTags, Log, TEXT("GameplayTagRedirects is in a deprecated location, after editing GameplayTags developer settings you must remove these manually"));
+			}
 
-			if (ensureMsgf(!TagRedirects.Contains(OldTagName), TEXT("Old tag %s is being redirected to more than one tag. Please remove all the redirections except for one."), *OldTagName.ToString()))
+			// Check settings object
+			for (const FGameplayTagRedirect& Redirect : MutableDefault->GameplayTagRedirects)
 			{
-				FGameplayTag OldTag = RequestGameplayTag(OldTagName, false); //< This only succeeds if OldTag is in the Table!
-				if (OldTag.IsValid())
+				FName OldTagName = Redirect.OldTagName;
+				FName NewTagName = Redirect.NewTagName;
+
+				if (ensureMsgf(!TagRedirects.Contains(OldTagName), TEXT("Old tag %s is being redirected to more than one tag. Please remove all the redirections except for one."), *OldTagName.ToString()))
 				{
-					FGameplayTagContainer MatchingChildren = RequestGameplayTagChildren(OldTag);
+					FGameplayTag OldTag = RequestGameplayTag(OldTagName, false); //< This only succeeds if OldTag is in the Table!
+					if (OldTag.IsValid())
+					{
+						FGameplayTagContainer MatchingChildren = RequestGameplayTagChildren(OldTag);
 
-					FString Msg = FString::Printf(TEXT("Old tag (%s) which is being redirected still exists in the table!  Generally you should "
-						TEXT("remove the old tags from the table when you are redirecting to new tags, or else users will ")
-						TEXT("still be able to add the old tags to containers.")), *OldTagName.ToString());
+						FString Msg = FString::Printf(TEXT("Old tag (%s) which is being redirected still exists in the table!  Generally you should "
+							TEXT("remove the old tags from the table when you are redirecting to new tags, or else users will ")
+							TEXT("still be able to add the old tags to containers.")), *OldTagName.ToString());
 
-					if (MatchingChildren.Num() == 0)
-					{
-						UE_LOG(LogGameplayTags, Warning, TEXT("%s"), *Msg);
-					}
-					else
-					{
-						Msg += TEXT("\nSuppressed warning due to redirected tag being a single component that matched other hierarchy elements.");
-						UE_LOG(LogGameplayTags, Log, TEXT("%s"), *Msg);
+						if (MatchingChildren.Num() == 0)
+						{
+							UE_LOG(LogGameplayTags, Warning, TEXT("%s"), *Msg);
+						}
+						else
+						{
+							Msg += TEXT("\nSuppressed warning due to redirected tag being a single component that matched other hierarchy elements.");
+							UE_LOG(LogGameplayTags, Log, TEXT("%s"), *Msg);
+						}
 					}
-				}
 
-				FGameplayTag NewTag = (NewTagName != NAME_None) ? RequestGameplayTag(NewTagName, false) : FGameplayTag();
+					FGameplayTag NewTag = (NewTagName != NAME_None) ? RequestGameplayTag(NewTagName, false) : FGameplayTag();
 
-				// Basic infinite recursion guard
-				int32 IterationsLeft = 10;
-				while (!NewTag.IsValid() && NewTagName != NAME_None)
-				{
-					bool bFoundRedirect = false;
-
-					// See if it got redirected again
-					for (const FGameplayTagRedirect& SecondRedirect : MutableDefault->GameplayTagRedirects)
+					// Basic infinite recursion guard
+					int32 IterationsLeft = 10;
+					while (!NewTag.IsValid() && NewTagName != NAME_None)
 					{
-						if (SecondRedirect.OldTagName == NewTagName)
+						bool bFoundRedirect = false;
+
+						// See if it got redirected again
+						for (const FGameplayTagRedirect& SecondRedirect : MutableDefault->GameplayTagRedirects)
 						{
-							NewTagName = SecondRedirect.NewTagName;
-							NewTag = RequestGameplayTag(NewTagName, false);
-							bFoundRedirect = true;
+							if (SecondRedirect.OldTagName == NewTagName)
+							{
+								NewTagName = SecondRedirect.NewTagName;
+								NewTag = RequestGameplayTag(NewTagName, false);
+								bFoundRedirect = true;
+								break;
+							}
+						}
+						IterationsLeft--;
+
+						if (!bFoundRedirect || IterationsLeft <= 0)
+						{
+							UE_LOG(LogGameplayTags, Warning, TEXT("Invalid new tag %s!  Cannot replace old tag %s."),
+								*Redirect.NewTagName.ToString(), *Redirect.OldTagName.ToString());
 							break;
 						}
 					}
-					IterationsLeft--;
 
-					if (!bFoundRedirect || IterationsLeft <= 0)
+					if (NewTag.IsValid())
 					{
-						UE_LOG(LogGameplayTags, Warning, TEXT("Invalid new tag %s!  Cannot replace old tag %s."),
-							*Redirect.NewTagName.ToString(), *Redirect.OldTagName.ToString());
-						break;
+						// Populate the map
+						TagRedirects.Add(OldTagName, NewTag);
 					}
 				}
-
-				if (NewTag.IsValid())
-				{
-					// Populate the map
-					TagRedirects.Add(OldTagName, NewTag);
-				}
 			}
 		}
 	}
@@ -454,8 +476,6 @@ void UGameplayTagsManager::ConstructNetIndex()
 		NetworkGameplayTagNodeIndex.SetNum(INVALID_TAGNETINDEX - 1);
 	}
 
-
-
 	UE_CLOG(PrintNetIndiceAssignment, LogGameplayTags, Display, TEXT("Assigning NetIndices to %d tags."), NetworkGameplayTagNodeIndex.Num() );
 
 	for (FGameplayTagNetIndex i = 0; i < NetworkGameplayTagNodeIndex.Num(); i++)
@@ -499,11 +519,10 @@ FGameplayTagNetIndex UGameplayTagsManager::GetNetIndexFromTag(const FGameplayTag
 bool UGameplayTagsManager::ShouldImportTagsFromINI() const
 {
 	UGameplayTagsSettings* MutableDefault = GetMutableDefault<UGameplayTagsSettings>();
-	FString DefaultEnginePath = FString::Printf(TEXT("%sDefaultEngine.ini"), *FPaths::SourceConfigDir());
-
+	
 	// Deprecated path
 	bool ImportFromINI = false;
-	if (GConfig->GetBool(TEXT("GameplayTags"), TEXT("ImportTagsFromConfig"), ImportFromINI, DefaultEnginePath))
+	if (GConfig->GetBool(TEXT("GameplayTags"), TEXT("ImportTagsFromConfig"), ImportFromINI, GEngineIni))
 	{
 		if (ImportFromINI)
 		{
@@ -662,27 +681,35 @@ bool UGameplayTagsManager::ImportSingleGameplayTag(FGameplayTag& Tag, FName Impo
 void UGameplayTagsManager::InitializeManager()
 {
 	check(!SingletonManager);
+	SCOPE_LOG_TIME_IN_SECONDS(TEXT("UGameplayTagsManager::InitializeManager"), nullptr);
 
 	SingletonManager = NewObject<UGameplayTagsManager>(GetTransientPackage(), NAME_None);
 	SingletonManager->AddToRoot();
 
-	UGameplayTagsSettings* MutableDefault = GetMutableDefault<UGameplayTagsSettings>();
-	FString DefaultEnginePath = FString::Printf(TEXT("%sDefaultEngine.ini"), *FPaths::SourceConfigDir());
-
-	TArray<FString> GameplayTagTables;
-	GConfig->GetArray(TEXT("GameplayTags"), TEXT("+GameplayTagTableList"), GameplayTagTables, DefaultEnginePath);
+	UGameplayTagsSettings* MutableDefault = nullptr;
+	{
+		SCOPE_LOG_GAMEPLAYTAGS(TEXT("UGameplayTagsManager::InitializeManager: Load settings"));
+		MutableDefault = GetMutableDefault<UGameplayTagsSettings>();
+	}
 
-	// Report deprecation
-	if (GameplayTagTables.Num() > 0)
 	{
-		UE_LOG(LogGameplayTags, Log, TEXT("GameplayTagTableList is in a deprecated location, open and save GameplayTag settings to fix"));
-		for (const FString& DataTable : GameplayTagTables)
+		SCOPE_LOG_GAMEPLAYTAGS(TEXT("UGameplayTagsManager::InitializeManager: Load deprecated"));
+
+		TArray<FString> GameplayTagTablePaths;
+		GConfig->GetArray(TEXT("GameplayTags"), TEXT("+GameplayTagTableList"), GameplayTagTablePaths, GEngineIni);
+
+		// Report deprecation
+		if (GameplayTagTablePaths.Num() > 0)
 		{
-			MutableDefault->GameplayTagTableList.AddUnique(DataTable);
+			UE_LOG(LogGameplayTags, Log, TEXT("GameplayTagTableList is in a deprecated location, open and save GameplayTag settings to fix"));
+			for (const FString& DataTable : GameplayTagTablePaths)
+			{
+				MutableDefault->GameplayTagTableList.AddUnique(DataTable);
+			}
 		}
 	}
 
-	SingletonManager->LoadGameplayTagTables();
+	SingletonManager->LoadGameplayTagTables(true);
 	SingletonManager->ConstructGameplayTagTree();
 
 	// Bind to end of engine init to be done adding native tags
@@ -722,19 +749,73 @@ void UGameplayTagsManager::AddTagTableRow(const FGameplayTagTableRow& TagRow, FN
 		bAllowNonRestrictedChildren = RestrictedTagRow->bAllowNonRestrictedChildren;
 	}
 
-	// Split the tag text on the "." delimiter to establish tag depth and then insert each tag into the
-	// gameplay tag tree
+	// Split the tag text on the "." delimiter to establish tag depth and then insert each tag into the gameplay tag tree
+	// We try to avoid as many FString->FName conversions as possible as they are slow
+	FName OriginalTagName = TagRow.Tag;
+	FString FullTagString = OriginalTagName.ToString();
+
+#if WITH_EDITOR
+	{
+		// In editor builds, validate string
+		// These must get fixed up cooking to work properly
+		FText ErrorText;
+		FString FixedString;
+
+		if (!IsValidGameplayTagString(FullTagString, &ErrorText, &FixedString))
+		{
+			if (FixedString.IsEmpty())
+			{
+				// No way to fix it
+				UE_LOG(LogGameplayTags, Error, TEXT("Invalid tag %s from source %s: %s!"), *FullTagString, *SourceName.ToString(), *ErrorText.ToString());
+				return;
+			}
+			else
+			{
+				UE_LOG(LogGameplayTags, Error, TEXT("Invalid tag %s from source %s: %s! Replacing with %s, you may need to modify InvalidTagCharacters"), *FullTagString, *SourceName.ToString(), *ErrorText.ToString(), *FixedString);
+				FullTagString = FixedString;
+				OriginalTagName = FName(*FixedString);
+			}
+		}
+	}
+#endif
+
 	TArray<FString> SubTags;
-	TagRow.Tag.ToString().ParseIntoArray(SubTags, TEXT("."), true);
+	FullTagString.ParseIntoArray(SubTags, TEXT("."), true);
+
+	// We will build this back up as we go
+	FullTagString.Reset();
 
+	int32 NumSubTags = SubTags.Num();
 	bool bHasSeenConflict = false;
 
-	for (int32 SubTagIdx = 0; SubTagIdx < SubTags.Num(); ++SubTagIdx)
+	for (int32 SubTagIdx = 0; SubTagIdx < NumSubTags; ++SubTagIdx)
 	{
-		TArray< TSharedPtr<FGameplayTagNode> >& ChildTags = CurNode.Get()->GetChildTagNodes();
+		bool bIsExplicitTag = (SubTagIdx == (NumSubTags - 1));
+		FName ShortTagName = *SubTags[SubTagIdx];
+		FName FullTagName;
 
-		bool bFromDictionary = (SubTagIdx == (SubTags.Num() - 1));
-		int32 InsertionIdx = InsertTagIntoNodeArray(*SubTags[SubTagIdx], CurNode, ChildTags, SourceName, TagRow.DevComment, bFromDictionary, bIsRestrictedTag, bAllowNonRestrictedChildren);
+		if (bIsExplicitTag)
+		{
+			// We already know the final name
+			FullTagName = OriginalTagName;
+		}
+		else if (SubTagIdx == 0)
+		{
+			// Full tag is the same as short tag, and start building full tag string
+			FullTagName = ShortTagName;
+			FullTagString = SubTags[SubTagIdx];
+		}
+		else
+		{
+			// Add .Tag and use that as full tag
+			FullTagString += TEXT(".");
+			FullTagString += SubTags[SubTagIdx];
+
+			FullTagName = FName(*FullTagString);
+		}
+			
+		TArray< TSharedPtr<FGameplayTagNode> >& ChildTags = CurNode.Get()->GetChildTagNodes();
+		int32 InsertionIdx = InsertTagIntoNodeArray(ShortTagName, FullTagName, CurNode, ChildTags, SourceName, TagRow.DevComment, bIsExplicitTag, bIsRestrictedTag, bAllowNonRestrictedChildren);
 
 		CurNode = ChildTags[InsertionIdx];
 
@@ -745,7 +826,7 @@ void UGameplayTagsManager::AddTagTableRow(const FGameplayTagTableRow& TagRow, FN
 			CurNode->bAncestorHasConflict = bHasSeenConflict;
 
 			// If the sources don't match and the tag is explicit and we should've added the tag explicitly here, we have a conflict
-			if (CurNode->SourceName != SourceName && (CurNode->bIsExplicitTag && SubTagIdx == SubTags.Num() - 1))
+			if (CurNode->SourceName != SourceName && (CurNode->bIsExplicitTag && bIsExplicitTag))
 			{
 				// mark all ancestors as having a bad descendant
 				for (TSharedPtr<FGameplayTagNode> CurAncestorNode : AncestorNodes)
@@ -811,20 +892,21 @@ bool UGameplayTagsManager::IsNativelyAddedTag(FGameplayTag Tag) const
 	return NativeTagsToAdd.Contains(Tag.GetTagName());
 }
 
-int32 UGameplayTagsManager::InsertTagIntoNodeArray(FName Tag, TSharedPtr<FGameplayTagNode> ParentNode, TArray< TSharedPtr<FGameplayTagNode> >& NodeArray, FName SourceName, const FString& DevComment, bool bIsExplicitTag, bool bIsRestrictedTag, bool bAllowNonRestrictedChildren)
+int32 UGameplayTagsManager::InsertTagIntoNodeArray(FName Tag, FName FullTag, TSharedPtr<FGameplayTagNode> ParentNode, TArray< TSharedPtr<FGameplayTagNode> >& NodeArray, FName SourceName, const FString& DevComment, bool bIsExplicitTag, bool bIsRestrictedTag, bool bAllowNonRestrictedChildren)
 {
-	int32 InsertionIdx = INDEX_NONE;
+	int32 FoundNodeIdx = INDEX_NONE;
 	int32 WhereToInsert = INDEX_NONE;
 
 	// See if the tag is already in the array
 	for (int32 CurIdx = 0; CurIdx < NodeArray.Num(); ++CurIdx)
 	{
-		if (NodeArray[CurIdx].IsValid())
+		FGameplayTagNode* CurrNode = NodeArray[CurIdx].Get();
+		if (CurrNode)
 		{
-			FGameplayTagNode* CurrNode = NodeArray[CurIdx].Get();
-			if (CurrNode->GetSimpleTagName() == Tag)
+			FName SimpleTagName = CurrNode->GetSimpleTagName();
+			if (SimpleTagName == Tag)
 			{
-				InsertionIdx = CurIdx;
+				FoundNodeIdx = CurIdx;
 #if WITH_EDITORONLY_DATA
 				// If we are explicitly adding this tag then overwrite the existing children restrictions with whatever is in the ini
 				// If we restrict children in the input data, make sure we restrict them in the existing node. This applies to explicit and implicitly defined nodes
@@ -849,7 +931,7 @@ int32 UGameplayTagsManager::InsertTagIntoNodeArray(FName Tag, TSharedPtr<FGamepl
 #endif				
 				break;
 			}
-			else if (CurrNode->GetSimpleTagName() > Tag && WhereToInsert == INDEX_NONE)
+			else if (SimpleTagName > Tag && WhereToInsert == INDEX_NONE)
 			{
 				// Insert new node before this
 				WhereToInsert = CurIdx;
@@ -857,7 +939,7 @@ int32 UGameplayTagsManager::InsertTagIntoNodeArray(FName Tag, TSharedPtr<FGamepl
 		}
 	}
 
-	if (InsertionIdx == INDEX_NONE)
+	if (FoundNodeIdx == INDEX_NONE)
 	{
 		if (WhereToInsert == INDEX_NONE)
 		{
@@ -866,13 +948,16 @@ int32 UGameplayTagsManager::InsertTagIntoNodeArray(FName Tag, TSharedPtr<FGamepl
 		}
 
 		// Don't add the root node as parent
-		TSharedPtr<FGameplayTagNode> TagNode = MakeShareable(new FGameplayTagNode(Tag, ParentNode != GameplayRootTag ? ParentNode : nullptr, bIsExplicitTag, bIsRestrictedTag, bAllowNonRestrictedChildren));
+		TSharedPtr<FGameplayTagNode> TagNode = MakeShareable(new FGameplayTagNode(Tag, FullTag, ParentNode != GameplayRootTag ? ParentNode : nullptr, bIsExplicitTag, bIsRestrictedTag, bAllowNonRestrictedChildren));
 
 		// Add at the sorted location
-		InsertionIdx = NodeArray.Insert(TagNode, WhereToInsert);
+		FoundNodeIdx = NodeArray.Insert(TagNode, WhereToInsert);
 
 		FGameplayTag GameplayTag = TagNode->GetCompleteTag();
 
+		// These should always match
+		ensure(GameplayTag.GetTagName() == FullTag);
+
 		{
 #if WITH_EDITOR
 			// This critical section is to handle an editor-only issue where tag requests come from another thread when async loading from a background thread in FGameplayTagContainer::Serialize.
@@ -887,23 +972,23 @@ int32 UGameplayTagsManager::InsertTagIntoNodeArray(FName Tag, TSharedPtr<FGamepl
 	static FName NativeSourceName = FGameplayTagSource::GetNativeName();
 
 	// Set/update editor only data
-	if (NodeArray[InsertionIdx]->SourceName.IsNone() && !SourceName.IsNone())
+	if (NodeArray[FoundNodeIdx]->SourceName.IsNone() && !SourceName.IsNone())
 	{
-		NodeArray[InsertionIdx]->SourceName = SourceName;
+		NodeArray[FoundNodeIdx]->SourceName = SourceName;
 	}
 	else if (SourceName == NativeSourceName)
 	{
 		// Native overrides other types
-		NodeArray[InsertionIdx]->SourceName = SourceName;
+		NodeArray[FoundNodeIdx]->SourceName = SourceName;
 	}
 
-	if (NodeArray[InsertionIdx]->DevComment.IsEmpty() && !DevComment.IsEmpty())
+	if (NodeArray[FoundNodeIdx]->DevComment.IsEmpty() && !DevComment.IsEmpty())
 	{
-		NodeArray[InsertionIdx]->DevComment = DevComment;
+		NodeArray[FoundNodeIdx]->DevComment = DevComment;
 	}
 #endif
 
-	return InsertionIdx;
+	return FoundNodeIdx;
 }
 
 void UGameplayTagsManager::PrintReplicationIndices()
@@ -1241,7 +1326,7 @@ bool UGameplayTagsManager::GetTagEditorData(FName TagName, FString& OutComment,
 void UGameplayTagsManager::EditorRefreshGameplayTagTree()
 {
 	DestroyGameplayTagTree();
-	LoadGameplayTagTables();
+	LoadGameplayTagTables(false);
 	ConstructGameplayTagTree();
 
 	OnEditorRefreshGameplayTagTree.Broadcast();
@@ -1438,6 +1523,75 @@ FGameplayTag UGameplayTagsManager::RequestGameplayTag(FName TagName, bool ErrorI
 	return FGameplayTag();
 }
 
+bool UGameplayTagsManager::IsValidGameplayTagString(const FString& TagString, FText* OutError, FString* OutFixedString)
+{
+	bool bIsValid = true;
+	FString FixedString = TagString;
+	FText ErrorText;
+
+	if (FixedString.IsEmpty())
+	{
+		ErrorText = LOCTEXT("EmptyStringError", "Tag is empty");
+		bIsValid = false;
+	}
+
+	while (FixedString.StartsWith(TEXT("."), ESearchCase::CaseSensitive))
+	{
+		ErrorText = LOCTEXT("StartWithPeriod", "Tag starts with .");
+		FixedString.RemoveAt(0);
+		bIsValid = false;
+	}
+
+	while (FixedString.EndsWith(TEXT("."), ESearchCase::CaseSensitive))
+	{
+		ErrorText = LOCTEXT("EndWithPeriod", "Tag ends with .");
+		FixedString.RemoveAt(FixedString.Len() - 1);
+		bIsValid = false;
+	}
+
+	while (FixedString.StartsWith(TEXT(" "), ESearchCase::CaseSensitive))
+	{
+		ErrorText = LOCTEXT("StartWithSpace", "Tag starts with space");
+		FixedString.RemoveAt(0);
+		bIsValid = false;
+	}
+
+	while (FixedString.EndsWith(TEXT(" "), ESearchCase::CaseSensitive))
+	{
+		ErrorText = LOCTEXT("EndWithSpace", "Tag ends with space");
+		FixedString.RemoveAt(FixedString.Len() - 1);
+		bIsValid = false;
+	}
+
+	FText TagContext = LOCTEXT("GameplayTagContext", "Tag");
+	if (!FName::IsValidXName(TagString, InvalidTagCharacters, &ErrorText, &TagContext))
+	{
+		for (TCHAR& TestChar : FixedString)
+		{
+			for (TCHAR BadChar : InvalidTagCharacters)
+			{
+				if (TestChar == BadChar)
+				{
+					TestChar = TEXT('_');
+				}
+			}
+		}
+
+		bIsValid = false;
+	}
+
+	if (OutError)
+	{
+		*OutError = ErrorText;
+	}
+	if (OutFixedString)
+	{
+		*OutFixedString = FixedString;
+	}
+
+	return bIsValid;
+}
+
 FGameplayTag UGameplayTagsManager::FindGameplayTagFromPartialString_Slow(FString PartialString) const
 {
 #if WITH_EDITOR
@@ -1533,7 +1687,7 @@ void UGameplayTagsManager::DoneAddingNativeTags()
 	// is initialized (DoneAddingNativeTags is bound to PostEngineInit to cover anything that's skipped).
 	if (GEngine && !bDoneAddingNativeTags)
 	{
-		UE_LOG(LogGameplayTags, Display, TEXT("UGameplayTagsManager::DoneAddingNativeTags. DelegateIsBound: %d"), (int32)OnLastChanceToAddNativeTags().IsBound());
+		UE_CLOG(GAMEPLAYTAGS_VERBOSE, LogGameplayTags, Display, TEXT("UGameplayTagsManager::DoneAddingNativeTags. DelegateIsBound: %d"), (int32)OnLastChanceToAddNativeTags().IsBound());
 		OnLastChanceToAddNativeTags().Broadcast();
 		bDoneAddingNativeTags = true;
 
@@ -1771,33 +1925,23 @@ bool FRestrictedGameplayTagTableRow::operator!=(FRestrictedGameplayTagTableRow c
 	return true;
 }
 
-FGameplayTagNode::FGameplayTagNode(FName InTag, TSharedPtr<FGameplayTagNode> InParentNode, bool InIsExplicitTag, bool InIsRestrictedTag, bool InAllowNonRestrictedChildren)
+FGameplayTagNode::FGameplayTagNode(FName InTag, FName InFullTag, TSharedPtr<FGameplayTagNode> InParentNode, bool InIsExplicitTag, bool InIsRestrictedTag, bool InAllowNonRestrictedChildren)
 	: Tag(InTag)
 	, ParentNode(InParentNode)
 	, NetIndex(INVALID_TAGNETINDEX)
-{
-	TArray<FGameplayTag> ParentCompleteTags;
-
-	TSharedPtr<FGameplayTagNode> CurNode = InParentNode;
+{	
+	// Manually construct the tag container as we want to bypass the safety checks
+	CompleteTagWithParents.GameplayTags.Add(FGameplayTag(InFullTag));
 
-	// Stop iterating at root node
-	while (CurNode.IsValid() && CurNode->GetSimpleTagName() != NAME_None)
+	FGameplayTagNode* RawParentNode = ParentNode.Get();
+	if (RawParentNode && RawParentNode->GetSimpleTagName() != NAME_None)
 	{
-		ParentCompleteTags.Add(CurNode->GetCompleteTag());
-		CurNode = CurNode->GetParentTagNode();
-	}
+		// Our parent nodes are already constructed, and must have it's tag in GameplayTags[0]
+		const FGameplayTagContainer ParentContainer = RawParentNode->GetSingleTagContainer();
 
-	FString CompleteTagString = InTag.ToString();
-
-	if (ParentCompleteTags.Num() > 0)
-	{
-		// If we have a parent, add parent., which will includes all earlier parents
-		CompleteTagString = FString::Printf(TEXT("%s.%s"), *ParentCompleteTags[0].ToString(), *InTag.ToString());
+		CompleteTagWithParents.ParentTags.Add(ParentContainer.GameplayTags[0]);
+		CompleteTagWithParents.ParentTags.Append(ParentContainer.ParentTags);
 	}
-	
-	// Manually construct the tag container as we want to bypass the safety checks
-	CompleteTagWithParents.GameplayTags.Add(FGameplayTag(FName(*CompleteTagString)));
-	CompleteTagWithParents.ParentTags = ParentCompleteTags;
 
 #if WITH_EDITORONLY_DATA
 	bIsExplicitTag = InIsExplicitTag;
diff --git a/Engine/Source/Runtime/GameplayTags/Private/GameplayTagsSettings.cpp b/Engine/Source/Runtime/GameplayTags/Private/GameplayTagsSettings.cpp
--- a/Engine/Source/Runtime/GameplayTags/Private/GameplayTagsSettings.cpp
+++ b/Engine/Source/Runtime/GameplayTags/Private/GameplayTagsSettings.cpp
@@ -62,6 +62,7 @@ UGameplayTagsSettings::UGameplayTagsSettings(const FObjectInitializer& ObjectIni
 	ImportTagsFromConfig = true;
 	WarnOnInvalidTags = true;
 	FastReplication = false;
+	InvalidTagCharacters = ("\"',");
 	NumBitsForContainerSize = 6;
 	NetIndexFirstBitSegment = 16;
 }
