Make GetDescription return a const reference instead of copy
#rnx

diff --git a/Engine/Source/Runtime/GameplayTags/Classes/GameplayTagContainer.h b/Engine/Source/Runtime/GameplayTags/Classes/GameplayTagContainer.h
--- a/Engine/Source/Runtime/GameplayTags/Classes/GameplayTagContainer.h
+++ b/Engine/Source/Runtime/GameplayTags/Classes/GameplayTagContainer.h
@@ -1045,7 +1045,7 @@ public:
 	void GetQueryExpr(struct FGameplayTagQueryExpression& OutExpr) const;
 
 	/** Returns description string. */
-	FString GetDescription() const { return UserDescription.IsEmpty() ? AutoDescription : UserDescription; };
+	const FString& GetDescription() const { return UserDescription.IsEmpty() ? AutoDescription : UserDescription; };
 
 #if WITH_EDITOR
 	/** Creates this query based on the given EditableQuery object */
