#UE4 Fixing a bug where gameplaytags are not properly redirected when using ImportText and adding an optional delegate on GameplayTagsManager for when a GameplayTag is loaded (either by serialization or ImportText)

diff --git a/Engine/Source/Runtime/GameplayTags/Classes/GameplayTagsManager.h b/Engine/Source/Runtime/GameplayTags/Classes/GameplayTagsManager.h
--- a/Engine/Source/Runtime/GameplayTags/Classes/GameplayTagsManager.h
+++ b/Engine/Source/Runtime/GameplayTags/Classes/GameplayTagsManager.h
@@ -483,6 +483,12 @@ class GAMEPLAYTAGS_API UGameplayTagsManager : public UObject
 	/** Returns a list of the owners for a restricted tag config file. May be empty */
 	void GetOwnersForTagSource(const FString& SourceName, TArray<FString>& OutOwners) const;
 
+	/** Notification that a tag container has been loaded via serialize */
+	void GameplayTagContainerLoaded(FGameplayTagContainer& Container, UProperty* SerializingProperty) const;
+
+	/** Notification that a gameplay tag has been loaded via serialize */
+	void SingleGameplayTagLoaded(FGameplayTag& Tag, UProperty* SerializingProperty) const;
+
 	/** Handles redirectors for an entire container, will also error on invalid tags */
 	void RedirectTagsForContainer(FGameplayTagContainer& Container, UProperty* SerializingProperty) const;
 
@@ -512,6 +518,9 @@ class GAMEPLAYTAGS_API UGameplayTagsManager : public UObject
 
 	bool IsNativelyAddedTag(FGameplayTag Tag) const;
 
+	DECLARE_MULTICAST_DELEGATE_OneParam(FOnGameplayTagLoaded, const FGameplayTag& /*Tag*/)
+	FOnGameplayTagLoaded OnGameplayTagLoadedDelegate;
+
 #if WITH_EDITOR
 	/** Gets a Filtered copy of the GameplayRootTags Array based on the comma delimited filter string passed in */
 	void GetFilteredGameplayRootTags(const FString& InFilterString, TArray< TSharedPtr<FGameplayTagNode> >& OutTagArray) const;
diff --git a/Engine/Source/Runtime/GameplayTags/Private/GameplayTagContainer.cpp b/Engine/Source/Runtime/GameplayTags/Private/GameplayTagContainer.cpp
--- a/Engine/Source/Runtime/GameplayTags/Private/GameplayTagContainer.cpp
+++ b/Engine/Source/Runtime/GameplayTags/Private/GameplayTagContainer.cpp
@@ -887,7 +887,7 @@ bool FGameplayTagContainer::Serialize(FArchive& Ar)
 		{
 			// Rename any tags that may have changed by the ini file.  Redirects can happen regardless of version.
 			// Regardless of version, want loading to have a chance to handle redirects
-			UGameplayTagsManager::Get().RedirectTagsForContainer(*this, Ar.GetSerializedProperty());
+			UGameplayTagsManager::Get().GameplayTagContainerLoaded(*this, Ar.GetSerializedProperty());
 		}
 
 		FillParentTags();
@@ -1308,7 +1308,7 @@ void FGameplayTag::PostSerialize(const FArchive& Ar)
 	if (Ar.IsLoading() && Ar.IsPersistent() && !(Ar.GetPortFlags() & PPF_Duplicate) && !(Ar.GetPortFlags() & PPF_DuplicateForPIE))
 	{
 		// Rename any tags that may have changed by the ini file.
-		UGameplayTagsManager::Get().RedirectSingleGameplayTag(*this, Ar.GetSerializedProperty());
+		UGameplayTagsManager::Get().SingleGameplayTagLoaded(*this, Ar.GetSerializedProperty());
 	}
 
 	if (Ar.IsSaving() && IsValid())
@@ -1338,7 +1338,10 @@ bool FGameplayTag::ImportTextItem(const TCHAR*& Buffer, int32 PortFlags, UObject
 	if (ImportedTag[0] == '(')
 	{
 		// Let normal ImportText handle this. It appears to be prepared for it.
-		return false;
+		UScriptStruct* ScriptStruct = FGameplayTag::StaticStruct();
+		Buffer = ScriptStruct->ImportText(Buffer, this, Parent, PortFlags, ErrorText, ScriptStruct->GetName(), false);
+		UGameplayTagsManager::Get().ImportSingleGameplayTag(*this, TagName);
+		return true;
 	}
 
 	return UGameplayTagsManager::Get().ImportSingleGameplayTag(*this, FName(*ImportedTag));
diff --git a/Engine/Source/Runtime/GameplayTags/Private/GameplayTagsManager.cpp b/Engine/Source/Runtime/GameplayTags/Private/GameplayTagsManager.cpp
--- a/Engine/Source/Runtime/GameplayTags/Private/GameplayTagsManager.cpp
+++ b/Engine/Source/Runtime/GameplayTags/Private/GameplayTagsManager.cpp
@@ -582,6 +582,26 @@ void UGameplayTagsManager::GetOwnersForTagSource(const FString& SourceName, TArr
 	}
 }
 
+void UGameplayTagsManager::GameplayTagContainerLoaded(FGameplayTagContainer& Container, UProperty* SerializingProperty) const
+{
+	RedirectTagsForContainer(Container, SerializingProperty);
+
+	if (OnGameplayTagLoadedDelegate.IsBound())
+	{
+		for (const FGameplayTag& Tag : Container)
+		{
+			OnGameplayTagLoadedDelegate.Broadcast(Tag);
+		}
+	}
+}
+
+void UGameplayTagsManager::SingleGameplayTagLoaded(FGameplayTag& Tag, UProperty* SerializingProperty) const
+{
+	RedirectSingleGameplayTag(Tag, SerializingProperty);
+
+	OnGameplayTagLoadedDelegate.Broadcast(Tag);
+}
+
 void UGameplayTagsManager::RedirectTagsForContainer(FGameplayTagContainer& Container, UProperty* SerializingProperty) const
 {
 	TSet<FName> NamesToRemove;
@@ -661,21 +681,30 @@ void UGameplayTagsManager::RedirectSingleGameplayTag(FGameplayTag& Tag, UPropert
 
 bool UGameplayTagsManager::ImportSingleGameplayTag(FGameplayTag& Tag, FName ImportedTagName) const
 {
+	bool bRetVal = false;
 	if (const FGameplayTag* RedirectedTag = TagRedirects.Find(ImportedTagName))
 	{
 		Tag = *RedirectedTag;
-		return true;
+		bRetVal = true;
 	}
 	else if (ValidateTagCreation(ImportedTagName))
 	{
 		// The tag name is valid
 		Tag.TagName = ImportedTagName;
-		return true;
+		bRetVal = true;
 	}
 
-	// No valid tag established in this attempt
-	Tag.TagName = NAME_None;
-	return false;
+	if (bRetVal)
+	{
+		OnGameplayTagLoadedDelegate.Broadcast(Tag);
+	}
+	else
+	{
+		// No valid tag established in this attempt
+		Tag.TagName = NAME_None;
+	}
+
+	return bRetVal;
 }
 
 void UGameplayTagsManager::InitializeManager()
