Improving tooltips for restricted gameplay tags.

diff --git a/Engine/Source/Runtime/GameplayTags/Classes/GameplayTagsSettings.h b/Engine/Source/Runtime/GameplayTags/Classes/GameplayTagsSettings.h
--- a/Engine/Source/Runtime/GameplayTags/Classes/GameplayTagsSettings.h
+++ b/Engine/Source/Runtime/GameplayTags/Classes/GameplayTagsSettings.h
@@ -150,7 +150,7 @@ class GAMEPLAYTAGS_API UGameplayTagsSettings : public UGameplayTagsList
 	UPROPERTY(config, EditAnywhere, Category = GameplayTags)
 	TArray<FGameplayTagRedirect> GameplayTagRedirects;
 
-	/** List of tags most frequently replicated */
+	/** List of most frequently replicated tags */
 	UPROPERTY(config, EditAnywhere, Category = "Advanced Replication")
 	TArray<FName> CommonlyReplicatedTags;
 
@@ -162,11 +162,15 @@ class GAMEPLAYTAGS_API UGameplayTagsSettings : public UGameplayTagsList
 	UPROPERTY(config, EditAnywhere, Category= "Advanced Replication")
 	int32 NetIndexFirstBitSegment;
 
-	/** Allows new tags to be saved into their own INI file. */
+	/** A list of .ini files used to store restricted gameplay tags. */
 	UPROPERTY(config, EditAnywhere, AdvancedDisplay, Category = "Advanced Gameplay Tags")
 	TArray<FRestrictedConfigInfo> RestrictedConfigFiles;
 #if WITH_EDITORONLY_DATA
 	// Dummy parameter used to hook the editor UI
+	/** Restricted Gameplay Tags.
+	 * 
+	 *  Restricted tags are intended to be top level tags that are important for your data hierarchy and modified by very few people.
+	 */
 	UPROPERTY(EditAnywhere, AdvancedDisplay, transient, Category = "Advanced Gameplay Tags")
 	FString RestrictedTagList;
 #endif
