Adding a note on how to filter gameplay tags to the header to make it easier to find.

diff --git a/Engine/Source/Runtime/GameplayTags/Classes/GameplayTagContainer.h b/Engine/Source/Runtime/GameplayTags/Classes/GameplayTagContainer.h
--- a/Engine/Source/Runtime/GameplayTags/Classes/GameplayTagContainer.h
+++ b/Engine/Source/Runtime/GameplayTags/Classes/GameplayTagContainer.h
@@ -44,7 +44,10 @@ enum class EGameplayContainerMatchType : uint8
 typedef uint16 FGameplayTagNetIndex;
 #define INVALID_TAGNETINDEX MAX_uint16
 
-/** A single gameplay tag, which represents a hierarchical name of the form x.y that is registered in the GameplayTagsManager */
+/**
+ * A single gameplay tag, which represents a hierarchical name of the form x.y that is registered in the GameplayTagsManager
+ * You can filter the gameplay tags displayed in the editor using, meta = (Categories = "Tag1.Tag2.Tag3"))
+ */
 USTRUCT(BlueprintType, meta = (HasNativeMake = "GameplayTags.BlueprintGameplayTagLibrary.MakeLiteralGameplayTag", HasNativeBreak = "GameplayTags.BlueprintGameplayTagLibrary.GetTagName"))
 struct GAMEPLAYTAGS_API FGameplayTag
 {
