Fix order in tag comparison. This makes gameplay tags show up correctly in alphabetical order.

diff --git a/Engine/Source/Runtime/GameplayTags/Private/GameplayTagsManager.cpp b/Engine/Source/Runtime/GameplayTags/Private/GameplayTagsManager.cpp
--- a/Engine/Source/Runtime/GameplayTags/Private/GameplayTagsManager.cpp
+++ b/Engine/Source/Runtime/GameplayTags/Private/GameplayTagsManager.cpp
@@ -995,7 +995,7 @@ int32 UGameplayTagsManager::InsertTagIntoNodeArray(FName Tag, FName FullTag, TSh
 #endif				
 				break;
 			}
-			else if (SimpleTagName.LexicalLess(Tag) && WhereToInsert == INDEX_NONE)
+			else if (Tag.LexicalLess(SimpleTagName) && WhereToInsert == INDEX_NONE)
 			{
 				// Insert new node before this
 				WhereToInsert = CurIdx;
