Fixed Gameplay Tag Query variables on components not being saved.

diff --git a/Engine/Source/Runtime/GameplayTags/Classes/GameplayTagContainer.h b/Engine/Source/Runtime/GameplayTags/Classes/GameplayTagContainer.h
--- a/Engine/Source/Runtime/GameplayTags/Classes/GameplayTagContainer.h
+++ b/Engine/Source/Runtime/GameplayTags/Classes/GameplayTagContainer.h
@@ -985,24 +985,26 @@ public:
 	FGameplayTagQuery& operator=(FGameplayTagQuery&& Other);
 
 private:
+	// Note: Properties need to be editable to allow FComponentPropertyWriter to serialize them, but are hidden in the editor by the customizations mentioned above.
+
 	/** Versioning for future token stream protocol changes. See EGameplayTagQueryStreamVersion. */
-	UPROPERTY()
+	UPROPERTY(EditAnywhere, Category = Hidden)
 	int32 TokenStreamVersion;
 
 	/** List of tags referenced by this entire query. Token stream stored indices into this list. */
-	UPROPERTY()
+	UPROPERTY(EditAnywhere, Category = Hidden)
 	TArray<FGameplayTag> TagDictionary;
 
 	/** Stream representation of the actual hierarchical query */
-	UPROPERTY()
+	UPROPERTY(EditAnywhere, Category = Hidden)
 	TArray<uint8> QueryTokenStream;
 
 	/** User-provided string describing the query */
-	UPROPERTY()
+	UPROPERTY(EditAnywhere, Category = Hidden)
 	FString UserDescription;
 
 	/** Auto-generated string describing the query */
-	UPROPERTY()
+	UPROPERTY(EditAnywhere, Category = Hidden)
 	FString AutoDescription;
 
 	/** Returns a gameplay tag from the tag dictionary */
