Convert GameplayTags serialization functions over to FStructuredArchive based versions

diff --git a/Engine/Source/Runtime/GameplayTags/Classes/GameplayTagContainer.h b/Engine/Source/Runtime/GameplayTags/Classes/GameplayTagContainer.h
--- a/Engine/Source/Runtime/GameplayTags/Classes/GameplayTagContainer.h
+++ b/Engine/Source/Runtime/GameplayTags/Classes/GameplayTagContainer.h
@@ -175,10 +175,9 @@ struct GAMEPLAYTAGS_API FGameplayTag
 		return TagName;
 	}
 
-	friend FArchive& operator<<(FArchive& Ar, FGameplayTag& GameplayTag)
+	friend void operator<<(FStructuredArchive::FSlot Slot, FGameplayTag& GameplayTag)
 	{
-		Ar << GameplayTag.TagName;
-		return Ar;
+		Slot << GameplayTag.TagName;
 	}
 
 	/** Overridden for fast serialize */
@@ -558,7 +557,7 @@ struct GAMEPLAYTAGS_API FGameplayTagContainer
 	void Reset(int32 Slack = 0);
 	
 	/** Serialize the tag container */
-	bool Serialize(FArchive& Ar);
+	bool Serialize(FStructuredArchive::FSlot Slot);
 
 	/** Efficient network serialize, takes advantage of the dictionary */
 	bool NetSerialize(FArchive& Ar, class UPackageMap* Map, bool& bOutSuccess);
@@ -867,7 +866,8 @@ struct TStructOpsTypeTraits<FGameplayTagContainer> : public TStructOpsTypeTraits
 {
 	enum
 	{
-		WithSerializer = true,
+		//WithStructuredSerializer = true,
+		WithStructuredSerializer = true,
 		WithIdenticalViaEquality = true,
 		WithNetSerializer = true,
 		WithNetSharedSerialization = true,
diff --git a/Engine/Source/Runtime/GameplayTags/Private/GameplayTagContainer.cpp b/Engine/Source/Runtime/GameplayTags/Private/GameplayTagContainer.cpp
--- a/Engine/Source/Runtime/GameplayTags/Private/GameplayTagContainer.cpp
+++ b/Engine/Source/Runtime/GameplayTags/Private/GameplayTagContainer.cpp
@@ -867,41 +867,43 @@ void FGameplayTagContainer::Reset(int32 Slack)
 	ParentTags.Reset(Slack);
 }
 
-bool FGameplayTagContainer::Serialize(FArchive& Ar)
+bool FGameplayTagContainer::Serialize(FStructuredArchive::FSlot Slot)
 {
-	const bool bOldTagVer = Ar.UE4Ver() < VER_UE4_GAMEPLAY_TAG_CONTAINER_TAG_TYPE_CHANGE;
+	FArchive& UnderlyingArchive = Slot.GetUnderlyingArchive();
+
+	const bool bOldTagVer = UnderlyingArchive.UE4Ver() < VER_UE4_GAMEPLAY_TAG_CONTAINER_TAG_TYPE_CHANGE;
 	
 	if (bOldTagVer)
 	{
 		TArray<FName> Tags_DEPRECATED;
-		Ar << Tags_DEPRECATED;
+		Slot << Tags_DEPRECATED;
 		// Too old to deal with
 		UE_LOG(LogGameplayTags, Error, TEXT("Failed to load old GameplayTag container, too old to migrate correctly"));
 	}
 	else
 	{
-		Ar << GameplayTags;
+		Slot << GameplayTags;
 	}
 	
 	// Only do redirects for real loads, not for duplicates or recompiles
-	if (Ar.IsLoading() )
+	if (UnderlyingArchive.IsLoading() )
 	{
-		if (Ar.IsPersistent() && !(Ar.GetPortFlags() & PPF_Duplicate) && !(Ar.GetPortFlags() & PPF_DuplicateForPIE))
+		if (UnderlyingArchive.IsPersistent() && !(UnderlyingArchive.GetPortFlags() & PPF_Duplicate) && !(UnderlyingArchive.GetPortFlags() & PPF_DuplicateForPIE))
 		{
 			// Rename any tags that may have changed by the ini file.  Redirects can happen regardless of version.
 			// Regardless of version, want loading to have a chance to handle redirects
-			UGameplayTagsManager::Get().GameplayTagContainerLoaded(*this, Ar.GetSerializedProperty());
+			UGameplayTagsManager::Get().GameplayTagContainerLoaded(*this, UnderlyingArchive.GetSerializedProperty());
 		}
 
 		FillParentTags();
 	}
 
-	if (Ar.IsSaving())
+	if (UnderlyingArchive.IsSaving())
 	{
 		// This marks the saved name for later searching
 		for (const FGameplayTag& Tag : GameplayTags)
 		{
-			Ar.MarkSearchableName(FGameplayTag::StaticStruct(), Tag.TagName);
+			UnderlyingArchive.MarkSearchableName(FGameplayTag::StaticStruct(), Tag.TagName);
 		}
 	}
 
