Back out changelists 6270254 and 6270367.

diff --git a/Engine/Source/Runtime/Core/Public/Math/Color.h b/Engine/Source/Runtime/Core/Public/Math/Color.h
--- a/Engine/Source/Runtime/Core/Public/Math/Color.h
+++ b/Engine/Source/Runtime/Core/Public/Math/Color.h
@@ -60,6 +60,18 @@ struct FLinearColor
 	CORE_API explicit FLinearColor(const FFloat16Color& C);
 
 	// Serializer.
+
+	friend FArchive& operator<<(FArchive& Ar,FLinearColor& Color)
+	{
+		return Ar << Color.R << Color.G << Color.B << Color.A;
+	}
+
+	bool Serialize( FArchive& Ar )
+	{
+		Ar << *this;
+		return true;
+	}
+
 	friend void operator<<(FStructuredArchive::FSlot Slot, FLinearColor& Color)
 	{
 		FStructuredArchive::FRecord Record = Slot.EnterRecord();
diff --git a/Engine/Source/Runtime/Core/Public/Math/Vector.h b/Engine/Source/Runtime/Core/Public/Math/Vector.h
--- a/Engine/Source/Runtime/Core/Public/Math/Vector.h
+++ b/Engine/Source/Runtime/Core/Public/Math/Vector.h
@@ -982,6 +982,20 @@ public:
 	 */
 	static CORE_API void GenerateClusterCenters(TArray<FVector>& Clusters, const TArray<FVector>& Points, int32 NumIterations, int32 NumConnectionsToBeValid);
 
+	/**
+	 * Serializer.
+	 *
+	 * @param Ar Serialization Archive.
+	 * @param V Vector to serialize.
+	 * @return Reference to Archive after serialization.
+	 */
+	friend FArchive& operator<<(FArchive& Ar, FVector& V)
+	{
+		// @warning BulkSerialize: FVector is serialized as memory dump
+		// See TArray::BulkSerialize for detailed description of implied limitations.
+		return Ar << V.X << V.Y << V.Z;
+	}
+
 	/**
 	 * Structured archive slot serializer.
 	 *
@@ -992,12 +1006,18 @@ public:
 	{
 		// @warning BulkSerialize: FVector is serialized as memory dump
 		// See TArray::BulkSerialize for detailed description of implied limitations.
-		FStructuredArchive::FRecord Record = Slot.EnterRecord();
-		Record << NAMED_ITEM("X", V.X);
-		Record << NAMED_ITEM("Y", V.Y);
-		Record << NAMED_ITEM("Z", V.Z);
+		FStructuredArchive::FStream Stream = Slot.EnterStream();
+		Stream.EnterElement() << V.X;
+		Stream.EnterElement() << V.Y;
+		Stream.EnterElement() << V.Z;
 	}
 	
+	bool Serialize(FArchive& Ar)
+	{
+		Ar << *this;
+		return true;
+	}
+
 	bool Serialize(FStructuredArchive::FSlot Slot)
 	{
 		Slot << *this;
diff --git a/Engine/Source/Runtime/CoreUObject/Private/UObject/Property.cpp b/Engine/Source/Runtime/CoreUObject/Private/UObject/Property.cpp
--- a/Engine/Source/Runtime/CoreUObject/Private/UObject/Property.cpp
+++ b/Engine/Source/Runtime/CoreUObject/Private/UObject/Property.cpp
@@ -30,7 +30,7 @@ struct TStructOpsTypeTraits<FVector> : public TStructOpsTypeTraitsBase2<FVector>
 		WithZeroConstructor = true,
 		WithNetSerializer = true,
 		WithNetSharedSerialization = true,
-		WithStructuredSerializer = true,
+		WithSerializer = true,
 	};
 };
 IMPLEMENT_STRUCT(Vector);
@@ -172,7 +172,7 @@ struct TStructOpsTypeTraits<FLinearColor> : public TStructOpsTypeTraitsBase2<FLi
 	{
 		WithNoInitConstructor = true,
 		WithZeroConstructor = true,
-		WithStructuredSerializer = true,
+		WithSerializer = true,
 	};
 };
 IMPLEMENT_STRUCT(LinearColor);
@@ -297,7 +297,7 @@ struct TStructOpsTypeTraits<FSoftObjectPath> : public TStructOpsTypeTraitsBase2<
 	enum
 	{
 		WithZeroConstructor = true,
-		WithStructuredSerializer = true,
+		WithSerializer = true,
 		WithCopy = true,
 		WithIdenticalViaEquality = true,
 		WithExportTextItem = true,
diff --git a/Engine/Source/Runtime/GameplayTags/Classes/GameplayTagContainer.h b/Engine/Source/Runtime/GameplayTags/Classes/GameplayTagContainer.h
--- a/Engine/Source/Runtime/GameplayTags/Classes/GameplayTagContainer.h
+++ b/Engine/Source/Runtime/GameplayTags/Classes/GameplayTagContainer.h
@@ -175,9 +175,10 @@ struct GAMEPLAYTAGS_API FGameplayTag
 		return TagName;
 	}
 
-	friend void operator<<(FStructuredArchive::FSlot Slot, FGameplayTag& GameplayTag)
+	friend FArchive& operator<<(FArchive& Ar, FGameplayTag& GameplayTag)
 	{
-		Slot << GameplayTag.TagName;
+		Ar << GameplayTag.TagName;
+		return Ar;
 	}
 
 	/** Overridden for fast serialize */
@@ -557,7 +558,7 @@ struct GAMEPLAYTAGS_API FGameplayTagContainer
 	void Reset(int32 Slack = 0);
 	
 	/** Serialize the tag container */
-	bool Serialize(FStructuredArchive::FSlot Slot);
+	bool Serialize(FArchive& Ar);
 
 	/** Efficient network serialize, takes advantage of the dictionary */
 	bool NetSerialize(FArchive& Ar, class UPackageMap* Map, bool& bOutSuccess);
@@ -866,8 +867,7 @@ struct TStructOpsTypeTraits<FGameplayTagContainer> : public TStructOpsTypeTraits
 {
 	enum
 	{
-		//WithStructuredSerializer = true,
-		WithStructuredSerializer = true,
+		WithSerializer = true,
 		WithIdenticalViaEquality = true,
 		WithNetSerializer = true,
 		WithNetSharedSerialization = true,
diff --git a/Engine/Source/Runtime/GameplayTags/Private/GameplayTagContainer.cpp b/Engine/Source/Runtime/GameplayTags/Private/GameplayTagContainer.cpp
--- a/Engine/Source/Runtime/GameplayTags/Private/GameplayTagContainer.cpp
+++ b/Engine/Source/Runtime/GameplayTags/Private/GameplayTagContainer.cpp
@@ -867,43 +867,41 @@ void FGameplayTagContainer::Reset(int32 Slack)
 	ParentTags.Reset(Slack);
 }
 
-bool FGameplayTagContainer::Serialize(FStructuredArchive::FSlot Slot)
+bool FGameplayTagContainer::Serialize(FArchive& Ar)
 {
-	FArchive& UnderlyingArchive = Slot.GetUnderlyingArchive();
-
-	const bool bOldTagVer = UnderlyingArchive.UE4Ver() < VER_UE4_GAMEPLAY_TAG_CONTAINER_TAG_TYPE_CHANGE;
+	const bool bOldTagVer = Ar.UE4Ver() < VER_UE4_GAMEPLAY_TAG_CONTAINER_TAG_TYPE_CHANGE;
 	
 	if (bOldTagVer)
 	{
 		TArray<FName> Tags_DEPRECATED;
-		Slot << Tags_DEPRECATED;
+		Ar << Tags_DEPRECATED;
 		// Too old to deal with
 		UE_LOG(LogGameplayTags, Error, TEXT("Failed to load old GameplayTag container, too old to migrate correctly"));
 	}
 	else
 	{
-		Slot << GameplayTags;
+		Ar << GameplayTags;
 	}
 	
 	// Only do redirects for real loads, not for duplicates or recompiles
-	if (UnderlyingArchive.IsLoading() )
+	if (Ar.IsLoading() )
 	{
-		if (UnderlyingArchive.IsPersistent() && !(UnderlyingArchive.GetPortFlags() & PPF_Duplicate) && !(UnderlyingArchive.GetPortFlags() & PPF_DuplicateForPIE))
+		if (Ar.IsPersistent() && !(Ar.GetPortFlags() & PPF_Duplicate) && !(Ar.GetPortFlags() & PPF_DuplicateForPIE))
 		{
 			// Rename any tags that may have changed by the ini file.  Redirects can happen regardless of version.
 			// Regardless of version, want loading to have a chance to handle redirects
-			UGameplayTagsManager::Get().GameplayTagContainerLoaded(*this, UnderlyingArchive.GetSerializedProperty());
+			UGameplayTagsManager::Get().GameplayTagContainerLoaded(*this, Ar.GetSerializedProperty());
 		}
 
 		FillParentTags();
 	}
 
-	if (UnderlyingArchive.IsSaving())
+	if (Ar.IsSaving())
 	{
 		// This marks the saved name for later searching
 		for (const FGameplayTag& Tag : GameplayTags)
 		{
-			UnderlyingArchive.MarkSearchableName(FGameplayTag::StaticStruct(), Tag.TagName);
+			Ar.MarkSearchableName(FGameplayTag::StaticStruct(), Tag.TagName);
 		}
 	}
 
