DR - Static analysis fixes
#rnx

diff --git a/Engine/Source/Runtime/GameplayTags/Private/GameplayTagsManager.cpp b/Engine/Source/Runtime/GameplayTags/Private/GameplayTagsManager.cpp
--- a/Engine/Source/Runtime/GameplayTags/Private/GameplayTagsManager.cpp
+++ b/Engine/Source/Runtime/GameplayTags/Private/GameplayTagsManager.cpp
@@ -61,7 +61,7 @@ void UGameplayTagsManager::LoadGameplayTagTables(bool bAllowAsyncLoad)
 	GameplayTagTables.Empty();
 
 	// If we're a cooked build and in a safe spot, start an async load so we can pipeline it
-	if (bAllowAsyncLoad && !WITH_EDITOR && !IsLoading() && MutableDefault->GameplayTagTableList.Num() > 0)
+	if (bAllowAsyncLoad && (!WITH_EDITOR) && !IsLoading() && MutableDefault->GameplayTagTableList.Num() > 0)
 	{
 		for (FSoftObjectPath DataTablePath : MutableDefault->GameplayTagTableList)
 		{
diff --git a/Engine/Source/Runtime/Renderer/Private/ScreenSpaceDiffuseIndirect.cpp b/Engine/Source/Runtime/Renderer/Private/ScreenSpaceDiffuseIndirect.cpp
--- a/Engine/Source/Runtime/Renderer/Private/ScreenSpaceDiffuseIndirect.cpp
+++ b/Engine/Source/Runtime/Renderer/Private/ScreenSpaceDiffuseIndirect.cpp
@@ -160,6 +160,7 @@ void RenderScreenSpaceDiffuseIndirect( FRHICommandListImmediate& RHICmdList, FVi
 
 			auto VertexShader	= View.ShaderMap->GetShader< FPostProcessVS >();
 			auto PixelShader	= View.ShaderMap->GetShader< FScreenSpaceDiffuseIndirectPS >( PermutationVector );
+			CA_ASSUME(PixelShader);
 
 			{
 				FGraphicsPipelineStateInitializer GraphicsPSOInit;
