Changed a few methods in FGameplayTagContainer to pass by reference instead of by value.

diff --git a/Engine/Source/Runtime/GameplayTags/Classes/GameplayTagContainer.h b/Engine/Source/Runtime/GameplayTags/Classes/GameplayTagContainer.h
--- a/Engine/Source/Runtime/GameplayTags/Classes/GameplayTagContainer.h
+++ b/Engine/Source/Runtime/GameplayTags/Classes/GameplayTagContainer.h
@@ -65,7 +65,7 @@ struct GAMEPLAYTAGS_API FGameplayTag
 	 * @param ErrorIfNotfound: ensure() that tag exists.
 	 * @return Will return the corresponding FGameplayTag or an empty one if not found.
 	 */
-	static FGameplayTag RequestGameplayTag(FName TagName, bool ErrorIfNotFound=true);
+	static FGameplayTag RequestGameplayTag(const FName& TagName, bool ErrorIfNotFound=true);
 
 	/** 
 	 * Returns true if this is a valid gameplay tag string (foo.bar.baz). If false, it will fill 
@@ -246,7 +246,7 @@ PRAGMA_ENABLE_DEPRECATION_WARNINGS
 protected:
 
 	/** Intentionally private so only the tag manager can use */
-	explicit FGameplayTag(FName InTagName);
+	explicit FGameplayTag(const FName& InTagName);
 
 	/** This Tags Name */
 	UPROPERTY(VisibleAnywhere, Category = GameplayTags, SaveGame)
@@ -547,14 +547,14 @@ struct GAMEPLAYTAGS_API FGameplayTagContainer
 	 * @param TagToRemove		Tag to remove from the container
 	 * @param bDeferParentTags	Skip calling FillParentTags for performance (must be handled by calling code)
 	 */
-	bool RemoveTag(FGameplayTag TagToRemove, bool bDeferParentTags=false);
+	bool RemoveTag(const FGameplayTag& TagToRemove, bool bDeferParentTags=false);
 
 	/**
 	 * Removes all tags in TagsToRemove from this container
 	 *
 	 * @param TagsToRemove	Tags to remove from the container
 	 */
-	void RemoveTags(FGameplayTagContainer TagsToRemove);
+	void RemoveTags(const FGameplayTagContainer& TagsToRemove);
 
 	/** Remove all tags from the container. Will maintain slack by default */
 	void Reset(int32 Slack = 0);
@@ -575,7 +575,7 @@ struct GAMEPLAYTAGS_API FGameplayTagContainer
 	FString ToString() const;
 
 	/** Sets from a ImportText string, used in asset registry */
-	void FromExportString(FString ExportString);
+	void FromExportString(const FString& ExportString);
 
 	/** Returns abbreviated human readable Tag list without parens or property names. If bQuoted is true it will quote each tag */
 	FString ToStringSimple(bool bQuoted = false) const;
diff --git a/Engine/Source/Runtime/GameplayTags/Private/GameplayTagContainer.cpp b/Engine/Source/Runtime/GameplayTags/Private/GameplayTagContainer.cpp
--- a/Engine/Source/Runtime/GameplayTags/Private/GameplayTagContainer.cpp
+++ b/Engine/Source/Runtime/GameplayTags/Private/GameplayTagContainer.cpp
@@ -821,7 +821,7 @@ bool FGameplayTagContainer::AddLeafTag(const FGameplayTag& TagToAdd)
 
 DECLARE_CYCLE_STAT(TEXT("FGameplayTagContainer::RemoveTag"), STAT_FGameplayTagContainer_RemoveTag, STATGROUP_GameplayTags);
 
-bool FGameplayTagContainer::RemoveTag(FGameplayTag TagToRemove, bool bDeferParentTags)
+bool FGameplayTagContainer::RemoveTag(const FGameplayTag& TagToRemove, bool bDeferParentTags)
 {
 	SCOPE_CYCLE_COUNTER(STAT_FGameplayTagContainer_RemoveTag);
 
@@ -841,7 +841,7 @@ bool FGameplayTagContainer::RemoveTag(FGameplayTag TagToRemove, bool bDeferParen
 
 DECLARE_CYCLE_STAT(TEXT("FGameplayTagContainer::RemoveTags"), STAT_FGameplayTagContainer_RemoveTags, STATGROUP_GameplayTags);
 
-void FGameplayTagContainer::RemoveTags(FGameplayTagContainer TagsToRemove)
+void FGameplayTagContainer::RemoveTags(const FGameplayTagContainer& TagsToRemove)
 {
 	SCOPE_CYCLE_COUNTER(STAT_FGameplayTagContainer_RemoveTags);
 
@@ -918,7 +918,7 @@ FString FGameplayTagContainer::ToString() const
 	return ExportString;
 }
 
-void FGameplayTagContainer::FromExportString(FString ExportString)
+void FGameplayTagContainer::FromExportString(const FString& ExportString)
 {
 	Reset();
 
@@ -1090,7 +1090,7 @@ const FGameplayTagContainer& FGameplayTag::GetSingleTagContainer() const
 	return FGameplayTagContainer::EmptyContainer;
 }
 
-FGameplayTag FGameplayTag::RequestGameplayTag(FName TagName, bool ErrorIfNotFound)
+FGameplayTag FGameplayTag::RequestGameplayTag(const FName& TagName, bool ErrorIfNotFound)
 {
 	return UGameplayTagsManager::Get().RequestGameplayTag(TagName, ErrorIfNotFound);
 }
@@ -1147,7 +1147,7 @@ int32 FGameplayTag::MatchesTagDepth(const FGameplayTag& TagToCheck) const
 	return UGameplayTagsManager::Get().GameplayTagsMatchDepth(*this, TagToCheck);
 }
 
-FGameplayTag::FGameplayTag(FName Name)
+FGameplayTag::FGameplayTag(const FName& Name)
 	: TagName(Name)
 {
 	// This constructor is used to bypass the table check and is only usable by GameplayTagManager
