#UE4 Do not rebuild gameplaytagtree if no inis were added in AddTagIniSearchPath

diff --git a/Engine/Source/Runtime/GameplayTags/Private/GameplayTagsManager.cpp b/Engine/Source/Runtime/GameplayTags/Private/GameplayTagsManager.cpp
--- a/Engine/Source/Runtime/GameplayTags/Private/GameplayTagsManager.cpp
+++ b/Engine/Source/Runtime/GameplayTags/Private/GameplayTagsManager.cpp
@@ -110,23 +110,26 @@ void UGameplayTagsManager::AddTagIniSearchPath(const FString& RootDir)
 	// Read all tags from the ini
 	TArray<FString> FilesInDirectory;
 	IFileManager::Get().FindFilesRecursive(FilesInDirectory, *RootDir, TEXT("*.ini"), true, false);
-	FilesInDirectory.Sort();
-	for (const FString& IniFilePath : FilesInDirectory)
+	if (FilesInDirectory.Num() > 0)
 	{
-		ExtraTagIniList.AddUnique(IniFilePath);
-	}
+		FilesInDirectory.Sort();
+		for (const FString& IniFilePath : FilesInDirectory)
+		{
+			ExtraTagIniList.AddUnique(IniFilePath);
+		}
 
-	if (!bIsConstructingGameplayTagTree)
-	{
+		if (!bIsConstructingGameplayTagTree)
+		{
 #if WITH_EDITOR
-		EditorRefreshGameplayTagTree();
+			EditorRefreshGameplayTagTree();
 #else
-		AddTagsFromAdditionalLooseIniFiles(FilesInDirectory);
+			AddTagsFromAdditionalLooseIniFiles(FilesInDirectory);
 
-		ConstructNetIndex();
+			ConstructNetIndex();
 
-		IGameplayTagsModule::OnGameplayTagTreeChanged.Broadcast();
+			IGameplayTagsModule::OnGameplayTagTreeChanged.Broadcast();
 #endif
+		}
 	}
 }
 
