#UE4 When tags are added after the disregardforgc window closes, the lists created that refer to those tags are now added to the root set to obey disregard for gc assumptions

diff --git a/Engine/Source/Runtime/GameplayTags/Private/GameplayTagsManager.cpp b/Engine/Source/Runtime/GameplayTags/Private/GameplayTagsManager.cpp
--- a/Engine/Source/Runtime/GameplayTags/Private/GameplayTagsManager.cpp
+++ b/Engine/Source/Runtime/GameplayTags/Private/GameplayTagsManager.cpp
@@ -8,6 +8,7 @@
 #include "Misc/ScopeLock.h"
 #include "Stats/StatsMisc.h"
 #include "Misc/ConfigCacheIni.h"
+#include "UObject/UObjectArray.h"
 #include "UObject/UObjectHash.h"
 #include "UObject/UObjectIterator.h"
 #include "UObject/LinkerLoad.h"
@@ -1559,11 +1560,19 @@ FGameplayTagSource* UGameplayTagsManager::FindOrAddTagSource(FName TagSourceName
 	{
 		NewSource->SourceTagList = NewObject<UGameplayTagsList>(this, TagSourceName, RF_Transient);
 		NewSource->SourceTagList->ConfigFileName = FString::Printf(TEXT("%sTags/%s"), *FPaths::SourceConfigDir(), *TagSourceName.ToString());
+		if (GUObjectArray.IsDisregardForGC(this))
+		{
+			NewSource->SourceTagList->AddToRoot();
+		}
 	}
 	else if (SourceType == EGameplayTagSourceType::RestrictedTagList)
 	{
 		NewSource->SourceRestrictedTagList = NewObject<URestrictedGameplayTagsList>(this, TagSourceName, RF_Transient);
 		NewSource->SourceRestrictedTagList->ConfigFileName = FString::Printf(TEXT("%sTags/%s"), *FPaths::SourceConfigDir(), *TagSourceName.ToString());
+		if (GUObjectArray.IsDisregardForGC(this))
+		{
+			NewSource->SourceRestrictedTagList->AddToRoot();
+		}
 	}
 
 	return NewSource;
