TBA: Convert a couple of types to using structured archive based serialization in tagged properties

#rnx

diff --git a/Engine/Source/Runtime/Core/Public/Math/Color.h b/Engine/Source/Runtime/Core/Public/Math/Color.h
--- a/Engine/Source/Runtime/Core/Public/Math/Color.h
+++ b/Engine/Source/Runtime/Core/Public/Math/Color.h
@@ -61,17 +61,6 @@ struct FLinearColor
 
 	// Serializer.
 
-	friend FArchive& operator<<(FArchive& Ar,FLinearColor& Color)
-	{
-		return Ar << Color.R << Color.G << Color.B << Color.A;
-	}
-
-	bool Serialize( FArchive& Ar )
-	{
-		Ar << *this;
-		return true;
-	}
-
 	friend void operator<<(FStructuredArchive::FSlot Slot, FLinearColor& Color)
 	{
 		FStructuredArchive::FRecord Record = Slot.EnterRecord();
diff --git a/Engine/Source/Runtime/Core/Public/Math/TwoVectors.h b/Engine/Source/Runtime/Core/Public/Math/TwoVectors.h
--- a/Engine/Source/Runtime/Core/Public/Math/TwoVectors.h
+++ b/Engine/Source/Runtime/Core/Public/Math/TwoVectors.h
@@ -7,6 +7,7 @@
 #include "Math/UnrealMathUtility.h"
 #include "Containers/UnrealString.h"
 #include "Math/Vector.h"
+#include "Serialization/ArchiveAdapters.h"
 
 /**
  * A pair of 3D vectors.
diff --git a/Engine/Source/Runtime/Core/Public/Math/Vector.h b/Engine/Source/Runtime/Core/Public/Math/Vector.h
--- a/Engine/Source/Runtime/Core/Public/Math/Vector.h
+++ b/Engine/Source/Runtime/Core/Public/Math/Vector.h
@@ -18,6 +18,7 @@
 #include "Internationalization/Internationalization.h"
 #include "Math/IntVector.h"
 #include "Math/Axis.h"
+#include "Serialization/ArchiveAdapters.h"
 
 #if PLATFORM_VECTOR_CUBIC_INTERP_SSE
 #include "Math/UnrealMathSSE.h"
@@ -982,20 +983,6 @@ public:
 	 */
 	static CORE_API void GenerateClusterCenters(TArray<FVector>& Clusters, const TArray<FVector>& Points, int32 NumIterations, int32 NumConnectionsToBeValid);
 
-	/**
-	 * Serializer.
-	 *
-	 * @param Ar Serialization Archive.
-	 * @param V Vector to serialize.
-	 * @return Reference to Archive after serialization.
-	 */
-	friend FArchive& operator<<(FArchive& Ar, FVector& V)
-	{
-		// @warning BulkSerialize: FVector is serialized as memory dump
-		// See TArray::BulkSerialize for detailed description of implied limitations.
-		return Ar << V.X << V.Y << V.Z;
-	}
-
 	/**
 	 * Structured archive slot serializer.
 	 *
@@ -1006,18 +993,12 @@ public:
 	{
 		// @warning BulkSerialize: FVector is serialized as memory dump
 		// See TArray::BulkSerialize for detailed description of implied limitations.
-		FStructuredArchive::FStream Stream = Slot.EnterStream();
-		Stream.EnterElement() << V.X;
-		Stream.EnterElement() << V.Y;
-		Stream.EnterElement() << V.Z;
+		FStructuredArchive::FRecord Record = Slot.EnterRecord();
+		Record << NAMED_ITEM("X", V.X);
+		Record << NAMED_ITEM("Y", V.Y);
+		Record << NAMED_ITEM("Z", V.Z);
 	}
 	
-	bool Serialize(FArchive& Ar)
-	{
-		Ar << *this;
-		return true;
-	}
-
 	bool Serialize(FStructuredArchive::FSlot Slot)
 	{
 		Slot << *this;
diff --git a/Engine/Source/Runtime/Core/Public/Serialization/ArchiveAdapters.h b/Engine/Source/Runtime/Core/Public/Serialization/ArchiveAdapters.h
--- a/Engine/Source/Runtime/Core/Public/Serialization/ArchiveAdapters.h
+++ b/Engine/Source/Runtime/Core/Public/Serialization/ArchiveAdapters.h
@@ -4,8 +4,11 @@
 
 #include "Concepts/Insertable.h"
 #include "Templates/Models.h"
-//#include "ArchiveFromStructuredArchive.h"
-//#include "StructuredArchiveFromArchive.h"
+
+#if WITH_TEXT_ARCHIVE_SUPPORT
+#include "ArchiveFromStructuredArchive.h"
+#endif
+#include "StructuredArchiveFromArchive.h"
 
 /**
  * Adapter operator which allows a type to stream to an FStructuredArchive::FSlot when it already supports streaming to an FArchive.
diff --git a/Engine/Source/Runtime/CoreUObject/Private/UObject/Property.cpp b/Engine/Source/Runtime/CoreUObject/Private/UObject/Property.cpp
--- a/Engine/Source/Runtime/CoreUObject/Private/UObject/Property.cpp
+++ b/Engine/Source/Runtime/CoreUObject/Private/UObject/Property.cpp
@@ -30,7 +30,7 @@ struct TStructOpsTypeTraits<FVector> : public TStructOpsTypeTraitsBase2<FVector>
 		WithZeroConstructor = true,
 		WithNetSerializer = true,
 		WithNetSharedSerialization = true,
-		WithSerializer = true,
+		WithStructuredSerializer = true,
 	};
 };
 IMPLEMENT_STRUCT(Vector);
@@ -172,7 +172,7 @@ struct TStructOpsTypeTraits<FLinearColor> : public TStructOpsTypeTraitsBase2<FLi
 	{
 		WithNoInitConstructor = true,
 		WithZeroConstructor = true,
-		WithSerializer = true,
+		WithStructuredSerializer = true,
 	};
 };
 IMPLEMENT_STRUCT(LinearColor);
@@ -297,7 +297,7 @@ struct TStructOpsTypeTraits<FSoftObjectPath> : public TStructOpsTypeTraitsBase2<
 	enum
 	{
 		WithZeroConstructor = true,
-		WithSerializer = true,
+		WithStructuredSerializer = true,
 		WithCopy = true,
 		WithIdenticalViaEquality = true,
 		WithExportTextItem = true,
diff --git a/Engine/Source/Runtime/GameplayTags/Classes/GameplayTagContainer.h b/Engine/Source/Runtime/GameplayTags/Classes/GameplayTagContainer.h
--- a/Engine/Source/Runtime/GameplayTags/Classes/GameplayTagContainer.h
+++ b/Engine/Source/Runtime/GameplayTags/Classes/GameplayTagContainer.h
@@ -175,10 +175,9 @@ struct GAMEPLAYTAGS_API FGameplayTag
 		return TagName;
 	}
 
-	friend FArchive& operator<<(FArchive& Ar, FGameplayTag& GameplayTag)
+	friend void operator<<(FStructuredArchive::FSlot Slot, FGameplayTag& GameplayTag)
 	{
-		Ar << GameplayTag.TagName;
-		return Ar;
+		Slot << GameplayTag.TagName;
 	}
 
 	/** Overridden for fast serialize */
@@ -558,7 +557,7 @@ struct GAMEPLAYTAGS_API FGameplayTagContainer
 	void Reset(int32 Slack = 0);
 	
 	/** Serialize the tag container */
-	bool Serialize(FArchive& Ar);
+	bool Serialize(FStructuredArchive::FSlot Slot);
 
 	/** Efficient network serialize, takes advantage of the dictionary */
 	bool NetSerialize(FArchive& Ar, class UPackageMap* Map, bool& bOutSuccess);
@@ -867,7 +866,7 @@ struct TStructOpsTypeTraits<FGameplayTagContainer> : public TStructOpsTypeTraits
 {
 	enum
 	{
-		WithSerializer = true,
+		WithStructuredSerializer = true,
 		WithIdenticalViaEquality = true,
 		WithNetSerializer = true,
 		WithNetSharedSerialization = true,
diff --git a/Engine/Source/Runtime/GameplayTags/Private/GameplayTagContainer.cpp b/Engine/Source/Runtime/GameplayTags/Private/GameplayTagContainer.cpp
--- a/Engine/Source/Runtime/GameplayTags/Private/GameplayTagContainer.cpp
+++ b/Engine/Source/Runtime/GameplayTags/Private/GameplayTagContainer.cpp
@@ -867,41 +867,43 @@ void FGameplayTagContainer::Reset(int32 Slack)
 	ParentTags.Reset(Slack);
 }
 
-bool FGameplayTagContainer::Serialize(FArchive& Ar)
+bool FGameplayTagContainer::Serialize(FStructuredArchive::FSlot Slot)
 {
-	const bool bOldTagVer = Ar.UE4Ver() < VER_UE4_GAMEPLAY_TAG_CONTAINER_TAG_TYPE_CHANGE;
+	FArchive& UnderlyingArchive = Slot.GetUnderlyingArchive();
+
+	const bool bOldTagVer = UnderlyingArchive.UE4Ver() < VER_UE4_GAMEPLAY_TAG_CONTAINER_TAG_TYPE_CHANGE;
 	
 	if (bOldTagVer)
 	{
 		TArray<FName> Tags_DEPRECATED;
-		Ar << Tags_DEPRECATED;
+		Slot << Tags_DEPRECATED;
 		// Too old to deal with
 		UE_LOG(LogGameplayTags, Error, TEXT("Failed to load old GameplayTag container, too old to migrate correctly"));
 	}
 	else
 	{
-		Ar << GameplayTags;
+		Slot << GameplayTags;
 	}
 	
 	// Only do redirects for real loads, not for duplicates or recompiles
-	if (Ar.IsLoading() )
+	if (UnderlyingArchive.IsLoading() )
 	{
-		if (Ar.IsPersistent() && !(Ar.GetPortFlags() & PPF_Duplicate) && !(Ar.GetPortFlags() & PPF_DuplicateForPIE))
+		if (UnderlyingArchive.IsPersistent() && !(UnderlyingArchive.GetPortFlags() & PPF_Duplicate) && !(UnderlyingArchive.GetPortFlags() & PPF_DuplicateForPIE))
 		{
 			// Rename any tags that may have changed by the ini file.  Redirects can happen regardless of version.
 			// Regardless of version, want loading to have a chance to handle redirects
-			UGameplayTagsManager::Get().GameplayTagContainerLoaded(*this, Ar.GetSerializedProperty());
+			UGameplayTagsManager::Get().GameplayTagContainerLoaded(*this, UnderlyingArchive.GetSerializedProperty());
 		}
 
 		FillParentTags();
 	}
 
-	if (Ar.IsSaving())
+	if (UnderlyingArchive.IsSaving())
 	{
 		// This marks the saved name for later searching
 		for (const FGameplayTag& Tag : GameplayTags)
 		{
-			Ar.MarkSearchableName(FGameplayTag::StaticStruct(), Tag.TagName);
+			UnderlyingArchive.MarkSearchableName(FGameplayTag::StaticStruct(), Tag.TagName);
 		}
 	}
 
