Fixed bug in FQueryEvaluator::ReadExpr that occurs when the expression isn't using tag data.

diff --git a/Engine/Source/Runtime/GameplayTags/Private/GameplayTagContainer.cpp b/Engine/Source/Runtime/GameplayTags/Private/GameplayTagContainer.cpp
--- a/Engine/Source/Runtime/GameplayTags/Private/GameplayTagContainer.cpp
+++ b/Engine/Source/Runtime/GameplayTags/Private/GameplayTagContainer.cpp
@@ -245,7 +245,7 @@ void FQueryEvaluator::ReadExpr(FGameplayTagQueryExpression& E)
 		{
 			FGameplayTagQueryExpression Exp;
 			ReadExpr(Exp);
-			Exp.AddExpr(Exp);
+			E.AddExpr(Exp);
 		}
 	}
 }
