Changed STATGROUP_GameplayTags to verbose to deflate the timing of scopes burdened by tag work.

diff --git a/Engine/Source/Runtime/GameplayTags/Classes/GameplayTagContainer.h b/Engine/Source/Runtime/GameplayTags/Classes/GameplayTagContainer.h
--- a/Engine/Source/Runtime/GameplayTags/Classes/GameplayTagContainer.h
+++ b/Engine/Source/Runtime/GameplayTags/Classes/GameplayTagContainer.h
@@ -15,7 +15,7 @@ struct FPropertyTag;
 
 GAMEPLAYTAGS_API DECLARE_LOG_CATEGORY_EXTERN(LogGameplayTags, Log, All);
 
-DECLARE_STATS_GROUP(TEXT("Gameplay Tags"), STATGROUP_GameplayTags, STATCAT_Advanced);
+DECLARE_STATS_GROUP_VERBOSE(TEXT("Gameplay Tags"), STATGROUP_GameplayTags, STATCAT_Advanced);
 
 DECLARE_CYCLE_STAT_EXTERN(TEXT("FGameplayTagContainer::HasTag"), STAT_FGameplayTagContainer_HasTag, STATGROUP_GameplayTags, GAMEPLAYTAGS_API);
 DECLARE_CYCLE_STAT_EXTERN(TEXT("FGameplayTagContainer::DoesTagContainerMatch"), STAT_FGameplayTagContainer_DoesTagContainerMatch, STATGROUP_GameplayTags, GAMEPLAYTAGS_API);
