#UE4 Read additional RestrictedTags lists from Tags ini files

diff --git a/Engine/Source/Runtime/GameplayTags/Classes/GameplayTagsManager.h b/Engine/Source/Runtime/GameplayTags/Classes/GameplayTagsManager.h
--- a/Engine/Source/Runtime/GameplayTags/Classes/GameplayTagsManager.h
+++ b/Engine/Source/Runtime/GameplayTags/Classes/GameplayTagsManager.h
@@ -693,6 +693,9 @@ private:
 	/** Constructs the net indices for each tag */
 	void ConstructNetIndex();
 
+	/** Reads the restricted config info from the specified ini */
+	void GetRestrictedConfigsFromIni(const FString& IniFilePath, TArray<struct FRestrictedConfigInfo>& OutRestrictedConfigs) const;
+
 	/** Marks all of the nodes that descend from CurNode as having an ancestor node that has a source conflict. */
 	void MarkChildrenOfNodeConflict(TSharedPtr<FGameplayTagNode> CurNode);
 
diff --git a/Engine/Source/Runtime/GameplayTags/Private/GameplayTagsManager.cpp b/Engine/Source/Runtime/GameplayTags/Private/GameplayTagsManager.cpp
--- a/Engine/Source/Runtime/GameplayTags/Private/GameplayTagsManager.cpp
+++ b/Engine/Source/Runtime/GameplayTags/Private/GameplayTagsManager.cpp
@@ -568,6 +568,22 @@ bool UGameplayTagsManager::ShouldImportTagsFromINI() const
 	return MutableDefault->ImportTagsFromConfig;
 }
 
+void UGameplayTagsManager::GetRestrictedConfigsFromIni(const FString& IniFilePath, TArray<FRestrictedConfigInfo>& OutRestrictedConfigs) const
+{
+	TArray<FString> IniConfigStrings;
+	if (GConfig->GetArray(TEXT("/Script/GameplayTags.GameplayTagsSettings"), TEXT("RestrictedConfigFiles"), IniConfigStrings, IniFilePath))
+	{
+		for (const FString& ConfigString : IniConfigStrings)
+		{
+			FRestrictedConfigInfo Config;
+			if (FRestrictedConfigInfo::StaticStruct()->ImportText(*ConfigString, &Config, nullptr, PPF_None, nullptr, FRestrictedConfigInfo::StaticStruct()->GetName()))
+			{
+				OutRestrictedConfigs.Add(Config);
+			}
+		}
+	}
+}
+
 void UGameplayTagsManager::GetRestrictedTagConfigFiles(TArray<FString>& RestrictedConfigFiles) const
 {
 	UGameplayTagsSettings* MutableDefault = GetMutableDefault<UGameplayTagsSettings>();
@@ -579,6 +595,16 @@ void UGameplayTagsManager::GetRestrictedTagConfigFiles(TArray<FString>& Restrict
 			RestrictedConfigFiles.Add(FString::Printf(TEXT("%sTags/%s"), *FPaths::SourceConfigDir(), *Config.RestrictedConfigName));
 		}
 	}
+
+	for (const FString& IniFilePath : ExtraTagIniList)
+	{
+		TArray<FRestrictedConfigInfo> IniRestrictedConfigs;
+		GetRestrictedConfigsFromIni(IniFilePath, IniRestrictedConfigs);
+		for (const FRestrictedConfigInfo& Config : IniRestrictedConfigs)
+		{
+			RestrictedConfigFiles.Add(FString::Printf(TEXT("%s/%s"), *FPaths::GetPath(IniFilePath), *Config.RestrictedConfigName));
+		}
+	}
 }
 
 void UGameplayTagsManager::GetRestrictedTagSources(TArray<const FGameplayTagSource*>& Sources) const
@@ -596,6 +622,20 @@ void UGameplayTagsManager::GetRestrictedTagSources(TArray<const FGameplayTagSour
 			}
 		}
 	}
+
+	for (const FString& IniFilePath : ExtraTagIniList)
+	{
+		TArray<FRestrictedConfigInfo> IniRestrictedConfigs;
+		GetRestrictedConfigsFromIni(IniFilePath, IniRestrictedConfigs);
+		for (const FRestrictedConfigInfo& Config : IniRestrictedConfigs)
+		{
+			const FGameplayTagSource* Source = FindTagSource(*Config.RestrictedConfigName);
+			if (Source)
+			{
+				Sources.Add(Source);
+			}
+		}
+	}
 }
 
 void UGameplayTagsManager::GetOwnersForTagSource(const FString& SourceName, TArray<FString>& OutOwners) const
