#UE4 RestrictedTags files discovered from ExtraTagIniList now properly add their tags in cooked builds.

diff --git a/Engine/Source/Runtime/GameplayTags/Classes/GameplayTagsManager.h b/Engine/Source/Runtime/GameplayTags/Classes/GameplayTagsManager.h
--- a/Engine/Source/Runtime/GameplayTags/Classes/GameplayTagsManager.h
+++ b/Engine/Source/Runtime/GameplayTags/Classes/GameplayTagsManager.h
@@ -668,6 +668,8 @@ private:
 
 	void AddChildrenTags(FGameplayTagContainer& TagContainer, TSharedPtr<FGameplayTagNode> GameplayTagNode, bool RecurseAll=true, bool OnlyIncludeDictionaryTags=false) const;
 
+	void AddRestrictedGameplayTagSource(const FString& FileName);
+
 	void AddTagsFromAdditionalLooseIniFiles(const TArray<FString>& IniFileList);
 
 	/**
diff --git a/Engine/Source/Runtime/GameplayTags/Private/GameplayTagsManager.cpp b/Engine/Source/Runtime/GameplayTags/Private/GameplayTagsManager.cpp
--- a/Engine/Source/Runtime/GameplayTags/Private/GameplayTagsManager.cpp
+++ b/Engine/Source/Runtime/GameplayTags/Private/GameplayTagsManager.cpp
@@ -173,6 +173,18 @@ void UGameplayTagsManager::AddTagIniSearchPath(const FString& RootDir)
 #if WITH_EDITOR
 			EditorRefreshGameplayTagTree();
 #else
+			for (const FString& IniFilePath : FilesInDirectory)
+			{
+				TArray<FRestrictedConfigInfo> IniRestrictedConfigs;
+				GetRestrictedConfigsFromIni(IniFilePath, IniRestrictedConfigs);
+				const FString IniDirectory = FPaths::GetPath(IniFilePath);
+				for (const FRestrictedConfigInfo& Config : IniRestrictedConfigs)
+				{
+					const FString RestrictedFileName = FString::Printf(TEXT("%s/%s"), *IniDirectory, *Config.RestrictedConfigName);
+					AddRestrictedGameplayTagSource(RestrictedFileName);
+				}
+			}
+
 			AddTagsFromAdditionalLooseIniFiles(FilesInDirectory);
 
 			ConstructNetIndex();
@@ -183,6 +195,36 @@ void UGameplayTagsManager::AddTagIniSearchPath(const FString& RootDir)
 	}
 }
 
+void UGameplayTagsManager::AddRestrictedGameplayTagSource(const FString& FileName)
+{
+	FName TagSource = FName(*FPaths::GetCleanFilename(FileName));
+	if (TagSource == NAME_None)
+	{
+		return;
+	}
+	RestrictedGameplayTagSourceNames.Add(TagSource);
+	FGameplayTagSource* FoundSource = FindOrAddTagSource(TagSource, EGameplayTagSourceType::RestrictedTagList);
+
+	// Make sure we have regular tag sources to match the restricted tag sources but don't try to read any tags from them yet.
+	FindOrAddTagSource(TagSource, EGameplayTagSourceType::TagList);
+
+	if (FoundSource && FoundSource->SourceRestrictedTagList)
+	{
+		FoundSource->SourceRestrictedTagList->LoadConfig(URestrictedGameplayTagsList::StaticClass(), *FileName);
+
+#if WITH_EDITOR
+		if (GIsEditor || IsRunningCommandlet()) // Sort tags for UI Purposes but don't sort in -game scenario since this would break compat with noneditor cooked builds
+		{
+			FoundSource->SourceRestrictedTagList->SortTags();
+		}
+#endif
+		for (const FRestrictedGameplayTagTableRow& TableRow : FoundSource->SourceRestrictedTagList->RestrictedGameplayTagList)
+		{
+			AddTagTableRow(TableRow, TagSource, true);
+		}
+	}
+}
+
 void UGameplayTagsManager::AddTagsFromAdditionalLooseIniFiles(const TArray<FString>& IniFileList)
 {
 	// Read all tags from the ini
@@ -259,32 +301,7 @@ void UGameplayTagsManager::ConstructGameplayTagTree()
 
 			for (const FString& FileName : RestrictedGameplayTagFiles)
 			{
-				FName TagSource = FName(*FPaths::GetCleanFilename(FileName));
-				if (TagSource == NAME_None)
-				{
-					continue;
-				}
-				RestrictedGameplayTagSourceNames.Add(TagSource);
-				FGameplayTagSource* FoundSource = FindOrAddTagSource(TagSource, EGameplayTagSourceType::RestrictedTagList);
-
-				// Make sure we have regular tag sources to match the restricted tag sources but don't try to read any tags from them yet.
-				FindOrAddTagSource(TagSource, EGameplayTagSourceType::TagList);
-
-				if (FoundSource && FoundSource->SourceRestrictedTagList)
-				{
-					FoundSource->SourceRestrictedTagList->LoadConfig(URestrictedGameplayTagsList::StaticClass(), *FileName);
-
-#if WITH_EDITOR
-					if (GIsEditor || IsRunningCommandlet()) // Sort tags for UI Purposes but don't sort in -game scenario since this would break compat with noneditor cooked builds
-					{
-						FoundSource->SourceRestrictedTagList->SortTags();
-					}
-#endif
-					for (const FRestrictedGameplayTagTableRow& TableRow : FoundSource->SourceRestrictedTagList->RestrictedGameplayTagList)
-					{
-						AddTagTableRow(TableRow, TagSource, true);
-					}
-				}
+				AddRestrictedGameplayTagSource(FileName);
 			}
 		}
 
