Remove several deprecated headers from 4.18 or earlier, these cause issues with symbol/file search
Remove gameplay framework and related CoreUObject code that was deprecated back in 4.18 or earlier
Properly deprecate soft object typedefs instead of removing them as they never had a full deprecation cycle
Skipped files that I know have pending changes from UE4 or appear to still be in use in some situations

diff --git a/Engine/Source/Runtime/CoreUObject/Private/Misc/RedirectCollector.cpp b/Engine/Source/Runtime/CoreUObject/Private/Misc/RedirectCollector.cpp
--- a/Engine/Source/Runtime/CoreUObject/Private/Misc/RedirectCollector.cpp
+++ b/Engine/Source/Runtime/CoreUObject/Private/Misc/RedirectCollector.cpp
@@ -47,25 +47,6 @@ void FRedirectCollector::OnSoftObjectPathLoaded(const FSoftObjectPath& InPath, F
 	SoftObjectPathMap.FindOrAdd(PackageName).Add(SoftObjectPathProperty);
 }
 
-void FRedirectCollector::OnStringAssetReferenceLoaded(const FString& InString)
-{
-	FSoftObjectPath Path(InString);
-	OnSoftObjectPathLoaded(Path, nullptr);
-}
-
-FString FRedirectCollector::OnStringAssetReferenceSaved(const FString& InString)
-{
-	FScopeLock ScopeLock(&CriticalSection);
-
-	FName Found = GetAssetPathRedirection(FName(*InString));
-
-	if (Found != NAME_None)
-	{
-		return Found.ToString();
-	}
-	return InString;
-}
-
 void FRedirectCollector::ResolveAllSoftObjectPaths(FName FilterPackage)
 {	
 	auto LoadSoftObjectPathLambda = [this](const FSoftObjectPathProperty& SoftObjectPathProperty)
diff --git a/Engine/Source/Runtime/CoreUObject/Private/UObject/CoreRedirects.cpp b/Engine/Source/Runtime/CoreUObject/Private/UObject/CoreRedirects.cpp
--- a/Engine/Source/Runtime/CoreUObject/Private/UObject/CoreRedirects.cpp
+++ b/Engine/Source/Runtime/CoreUObject/Private/UObject/CoreRedirects.cpp
@@ -1700,7 +1700,7 @@ static void RegisterNativeRedirects46(TArray<FCoreRedirect>& Redirects)
 	FUNCTION_REDIRECT("LightComponent.SetBrightness", "LightComponent.SetIntensity");
 	FUNCTION_REDIRECT("NavigationPath.GetPathLenght", "NavigationPath.GetPathLength");
 	FUNCTION_REDIRECT("Pawn.GetMovementInputVector", "Pawn.K2_GetMovementInputVector");
-	FUNCTION_REDIRECT("PawnMovementComponent.GetInputVector", "PawnMovementComponent.K2_GetInputVector");
+	FUNCTION_REDIRECT("PawnMovementComponent.GetInputVector", "PawnMovementComponent.GetPendingInputVector");
 	FUNCTION_REDIRECT("SceneComponent.AttachTo", "SceneComponent.K2_AttachTo");
 	FUNCTION_REDIRECT("SkyLightComponent.SetBrightness", "SkyLightComponent.SetIntensity");
 
diff --git a/Engine/Source/Runtime/CoreUObject/Private/UObject/PropertyStruct.cpp b/Engine/Source/Runtime/CoreUObject/Private/UObject/PropertyStruct.cpp
--- a/Engine/Source/Runtime/CoreUObject/Private/UObject/PropertyStruct.cpp
+++ b/Engine/Source/Runtime/CoreUObject/Private/UObject/PropertyStruct.cpp
@@ -259,12 +259,6 @@ FString FStructProperty::GetCPPMacroType( FString& ExtendedTypeText ) const
 	return TEXT("STRUCT");
 }
 
-void FStructProperty::ExportTextItem_Static(UScriptStruct* InStruct, FString& ValueStr, const void* PropertyValue, const void* DefaultValue, UObject* Parent, int32 PortFlags, UObject* ExportRootScope)
-{
-	// For backward compatibility skip the native export 
-	InStruct->ExportText(ValueStr, PropertyValue, DefaultValue, Parent, PortFlags, ExportRootScope, false);
-}
-
 void FStructProperty::ExportTextItem( FString& ValueStr, const void* PropertyValue, const void* DefaultValue, UObject* Parent, int32 PortFlags, UObject* ExportRootScope ) const
 {
 	Struct->ExportText(ValueStr, PropertyValue, DefaultValue, Parent, PortFlags, ExportRootScope, true);
@@ -302,11 +296,6 @@ const TCHAR* FStructProperty::ImportText_Internal(const TCHAR* InBuffer, void* D
 	return Result;
 }
 
-const TCHAR* FStructProperty::ImportText_Static(UScriptStruct* InStruct, const FString& Name, const TCHAR* InBuffer, void* Data, int32 PortFlags, UObject* Parent, FOutputDevice* ErrorText)
-{
-	return InStruct->ImportText(InBuffer, Data, Parent, PortFlags, ErrorText, Name, true);
-}
-
 void FStructProperty::CopyValuesInternal( void* Dest, void const* Src, int32 Count  ) const
 {
 	Struct->CopyScriptStruct(Dest, Src, Count);
diff --git a/Engine/Source/Runtime/CoreUObject/Public/Misc/RedirectCollector.h b/Engine/Source/Runtime/CoreUObject/Public/Misc/RedirectCollector.h
--- a/Engine/Source/Runtime/CoreUObject/Public/Misc/RedirectCollector.h
+++ b/Engine/Source/Runtime/CoreUObject/Public/Misc/RedirectCollector.h
@@ -102,18 +102,6 @@ public:
 		return SoftObjectPathMap.Num() > 0;
 	}
 
-	UE_DEPRECATED(4.17, "OnStringAssetReferenceSaved is deprecated, call GetAssetPathRedirection")
-	FString OnStringAssetReferenceSaved(const FString& InString);
-
-	UE_DEPRECATED(4.18, "OnStringAssetReferenceLoaded is deprecated, call OnSoftObjectPathLoaded")
-	void OnStringAssetReferenceLoaded(const FString& InString);
-
-	UE_DEPRECATED(4.18, "ResolveStringAssetReference is deprecated, call ResolveAllSoftObjectPaths")
-	void ResolveStringAssetReference(FName FilterPackage = NAME_None, bool bProcessAlreadyResolvedPackages = true)
-	{
-		ResolveAllSoftObjectPaths(FilterPackage);
-	}
-
 private:
 
 	/** A map of assets referenced by soft object paths, with the key being the package with the reference */
diff --git a/Engine/Source/Runtime/CoreUObject/Public/Misc/StringAssetReference.h b/Engine/Source/Runtime/CoreUObject/Public/Misc/StringAssetReference.h
deleted file mode 100644
diff --git a/Engine/Source/Runtime/CoreUObject/Public/Misc/StringClassReference.h b/Engine/Source/Runtime/CoreUObject/Public/Misc/StringClassReference.h
deleted file mode 100644
diff --git a/Engine/Source/Runtime/CoreUObject/Public/UObject/AssetPtr.h b/Engine/Source/Runtime/CoreUObject/Public/UObject/AssetPtr.h
deleted file mode 100644
diff --git a/Engine/Source/Runtime/CoreUObject/Public/UObject/SoftObjectPath.h b/Engine/Source/Runtime/CoreUObject/Public/UObject/SoftObjectPath.h
--- a/Engine/Source/Runtime/CoreUObject/Public/UObject/SoftObjectPath.h
+++ b/Engine/Source/Runtime/CoreUObject/Public/UObject/SoftObjectPath.h
@@ -332,11 +332,10 @@ private:
 	static FSoftObjectPath GetOrCreateIDForObject(const UObject *Object);
 };
 
-// Not deprecating these yet as it will lead to too many warnings in games
-//UE_DEPRECATED(4.18, "FStringAssetReference was renamed to FSoftObjectPath as it is now not always a string and can also refer to a subobject")
+UE_DEPRECATED(5.0, "FStringAssetReference was renamed to FSoftObjectPath as it is now not always a string and can also refer to a subobject")
 typedef FSoftObjectPath FStringAssetReference;
 
-//UE_DEPRECATED(4.18, "FStringClassReference was renamed to FSoftClassPath")
+UE_DEPRECATED(5.0, "FStringClassReference was renamed to FSoftClassPath")
 typedef FSoftClassPath FStringClassReference;
 
 /** Options for how to set soft object path collection */
diff --git a/Engine/Source/Runtime/CoreUObject/Public/UObject/SoftObjectPtr.h b/Engine/Source/Runtime/CoreUObject/Public/UObject/SoftObjectPtr.h
--- a/Engine/Source/Runtime/CoreUObject/Public/UObject/SoftObjectPtr.h
+++ b/Engine/Source/Runtime/CoreUObject/Public/UObject/SoftObjectPtr.h
@@ -65,12 +65,6 @@ public:
 		return GetUniqueID();
 	}
 
-	UE_DEPRECATED(4.18, "ToStringReference was renamed to ToSoftObjectPath")
-	FORCEINLINE const FSoftObjectPath& ToStringReference() const
-	{
-		return GetUniqueID();
-	}
-
 	/** Returns string representation of reference, in form /package/path.assetname */
 	FORCEINLINE FString ToString() const
 	{
@@ -344,12 +338,6 @@ public:
 		return SoftObjectPtr.GetUniqueID();
 	}
 
-	UE_DEPRECATED(4.18, "ToStringReference was renamed to ToSoftObjectPath")
-	FORCEINLINE const FSoftObjectPath& ToStringReference() const
-	{
-		return SoftObjectPtr.GetUniqueID();
-	}
-
 	/** Returns string representation of reference, in form /package/path.assetname */
 	FORCEINLINE FString ToString() const
 	{
@@ -559,12 +547,6 @@ public:
 		return SoftObjectPtr.GetUniqueID();
 	}
 
-	UE_DEPRECATED(4.18, "ToStringReference was renamed to ToSoftObjectPath")
-	FORCEINLINE const FSoftObjectPath& ToStringReference() const
-	{
-		return SoftObjectPtr.GetUniqueID();
-	}
-
 	/** Returns string representation of reference, in form /package/path.assetname  */
 	FORCEINLINE FString ToString() const
 	{
@@ -640,15 +622,9 @@ struct FSoftObjectPtrLexicalLess : private FSoftObjectPathLexicalLess
 	}
 };
 
-UE_DEPRECATED(4.18, "FAssetPtr was renamed to FSoftObjectPtr as it is not necessarily an asset")
-typedef FSoftObjectPtr FAssetPtr;
-
-// Not deprecating these yet as it will lead to too many warnings in games
-//UE_DEPRECATED(4.18, "TAssetPtr was renamed to TSoftObjectPtr as it is not necessarily an asset")
 template<class T=UObject>
-using TAssetPtr = TSoftObjectPtr<T>;
+using TAssetPtr UE_DEPRECATED(5.0, "TAssetPtr was renamed to TSoftObjectPtr as it is not necessarily an asset") = TSoftObjectPtr<T>;
 
-//UE_DEPRECATED(4.18, "TAssetSubclassOf was renamed to TSoftClassPtr")
 template<class TClass = UObject>
-using TAssetSubclassOf = TSoftClassPtr<TClass>;
+using TAssetSubclassOf UE_DEPRECATED(5.0, "TAssetSubclassOf was renamed to TSoftClassPtr") = TSoftClassPtr<TClass>;
 
diff --git a/Engine/Source/Runtime/CoreUObject/Public/UObject/UObject.h b/Engine/Source/Runtime/CoreUObject/Public/UObject/UObject.h
deleted file mode 100644
diff --git a/Engine/Source/Runtime/CoreUObject/Public/UObject/UnrealType.h b/Engine/Source/Runtime/CoreUObject/Public/UObject/UnrealType.h
--- a/Engine/Source/Runtime/CoreUObject/Public/UObject/UnrealType.h
+++ b/Engine/Source/Runtime/CoreUObject/Public/UObject/UnrealType.h
@@ -4782,12 +4782,6 @@ public:
 	virtual EConvertFromTypeResult ConvertFromType(const FPropertyTag& Tag, FStructuredArchive::FSlot Slot, uint8* Data, UStruct* DefaultsStruct) override;
 	// End of FProperty interface
 
-	UE_DEPRECATED(4.14, "Use UScriptStruct::ImportText instead")
-	static const TCHAR* ImportText_Static(UScriptStruct* InStruct, const FString& InName, const TCHAR* Buffer, void* Data, int32 PortFlags, UObject* OwnerObject, FOutputDevice* ErrorText);
-	
-	UE_DEPRECATED(4.14, "Use UScriptStruct::ExportText instead")
-	static void ExportTextItem_Static(class UScriptStruct* InStruct, FString& ValueStr, const void* PropertyValue, const void* DefaultValue, UObject* Parent, int32 PortFlags, UObject* ExportRootScope);
-
 	bool UseBinaryOrNativeSerialization(const FArchive& Ar) const;
 
 private:
diff --git a/Engine/Source/Runtime/Engine/Classes/Components/PrimitiveComponent.h b/Engine/Source/Runtime/Engine/Classes/Components/PrimitiveComponent.h
--- a/Engine/Source/Runtime/Engine/Classes/Components/PrimitiveComponent.h
+++ b/Engine/Source/Runtime/Engine/Classes/Components/PrimitiveComponent.h
@@ -1162,20 +1162,6 @@ public:
 	UFUNCTION(BlueprintCallable, Category="Physics", meta=(UnsafeDuringActorConstruction="true"))
 	virtual void AddImpulse(FVector Impulse, FName BoneName = NAME_None, bool bVelChange = false);
 
-	/**
-	*	Add an angular impulse to a single rigid body. Good for one time instant burst.
-	*
-	*	@param	AngularImpulse	Magnitude and direction of impulse to apply. Direction is axis of rotation.
-	*	@param	BoneName	If a SkeletalMeshComponent, name of body to apply angular impulse to. 'None' indicates root body.
-	*	@param	bVelChange	If true, the Strength is taken as a change in angular velocity instead of an impulse (ie. mass will have no effect).
-	*/
-	UE_DEPRECATED(4.18, "Use AddAngularImpulseInRadians instead.")
-	UFUNCTION(BlueprintCallable, Category = "Physics", meta=(UnsafeDuringActorConstruction="true", DeprecatedFunction, DeprecationMessage="Use AddAngularImpulseInRadians instead"))
-	virtual void AddAngularImpulse(FVector Impulse, FName BoneName = NAME_None, bool bVelChange = false)
-	{
-		AddAngularImpulseInRadians(Impulse, BoneName, bVelChange);
-	}
-
 	/**
 	*	Add an angular impulse to a single rigid body. Good for one time instant burst.
 	*
@@ -1266,19 +1252,6 @@ public:
 	UFUNCTION(BlueprintCallable, Category="Physics", meta=(UnsafeDuringActorConstruction="true"))
 	virtual void AddRadialForce(FVector Origin, float Radius, float Strength, enum ERadialImpulseFalloff Falloff, bool bAccelChange = false);
 
-	/**
-	 *	Add a torque to a single rigid body.
-	 *	@param Torque		Torque to apply. Direction is axis of rotation and magnitude is strength of torque.
-	 *	@param BoneName		If a SkeletalMeshComponent, name of body to apply torque to. 'None' indicates root body.
-	 *  @param bAccelChange If true, Torque is taken as a change in angular acceleration instead of a physical torque (i.e. mass will have no effect).
-	 */
-	UE_DEPRECATED(4.18, "Use AddTorqueInRadians instead.")
-	UFUNCTION(BlueprintCallable, Category="Physics", meta=(UnsafeDuringActorConstruction="true", DeprecatedFunction, DeprecationMessage="Use AddTorqueInRadians instead"))
-	void AddTorque(FVector Torque, FName BoneName = NAME_None, bool bAccelChange = false)
-	{
-		AddTorqueInRadians(Torque, BoneName, bAccelChange);
-	}
-
 	/**
 	 *	Add a torque to a single rigid body.
 	 *	@param Torque		Torque to apply. Direction is axis of rotation and magnitude is strength of torque.
@@ -1335,21 +1308,6 @@ public:
 	UFUNCTION(BlueprintCallable, Category="Physics", meta=(UnsafeDuringActorConstruction="true"))
 	virtual void SetAllPhysicsLinearVelocity(FVector NewVel, bool bAddToCurrent = false);
 
-	/**
-	 *	Set the angular velocity of a single body.
-	 *	This should be used cautiously - it may be better to use AddTorque or AddImpulse.
-	 *
-	 *	@param NewAngVel		New angular velocity to apply to body, in degrees per second.
-	 *	@param bAddToCurrent	If true, NewAngVel is added to the existing angular velocity of the body.
-	 *	@param BoneName			If a SkeletalMeshComponent, name of body to modify angular velocity of. 'None' indicates root body.
-	 */
-	UE_DEPRECATED(4.18, "Use SetPhysicsAngularVelocityInDegrees instead.")
-	UFUNCTION(BlueprintCallable, Category="Physics", meta=(UnsafeDuringActorConstruction="true", DeprecatedFunction, DeprecationMessage="Use SetPhysicsAngularVelocityInDegrees instead"))
-	void SetPhysicsAngularVelocity(FVector NewAngVel, bool bAddToCurrent = false, FName BoneName = NAME_None)
-	{
-		SetPhysicsAngularVelocityInDegrees(NewAngVel, bAddToCurrent, BoneName);
-	}
-
 	/**
 	 *	Set the angular velocity of a single body.
 	 *	This should be used cautiously - it may be better to use AddTorque or AddImpulse.
@@ -1375,20 +1333,6 @@ public:
 		SetPhysicsAngularVelocityInRadians(FMath::DegreesToRadians(NewAngVel), bAddToCurrent, BoneName);
 	}
 
-	/**
-	*	Set the maximum angular velocity of a single body.
-	*
-	*	@param NewMaxAngVel		New maximum angular velocity to apply to body, in degrees per second.
-	*	@param bAddToCurrent	If true, NewMaxAngVel is added to the existing maximum angular velocity of the body.
-	*	@param BoneName			If a SkeletalMeshComponent, name of body to modify maximum angular velocity of. 'None' indicates root body.
-	*/
-	UE_DEPRECATED(4.18, "Use SetPhysicsMaxAngularVelocityInDegrees instead.")
-	UFUNCTION(BlueprintCallable, Category = "Physics", meta=(UnsafeDuringActorConstruction="true", DeprecatedFunction, DeprecationMessage="Use SetPhysicsMaxAngularVelocityInDegrees instead"))
-	void SetPhysicsMaxAngularVelocity(float NewMaxAngVel, bool bAddToCurrent = false, FName BoneName = NAME_None)
-	{
-		SetPhysicsMaxAngularVelocityInDegrees(NewMaxAngVel, bAddToCurrent, BoneName);
-	}
-
 	/**
 	*	Set the maximum angular velocity of a single body.
 	*
@@ -1412,17 +1356,6 @@ public:
 	UFUNCTION(BlueprintCallable, Category = "Physics", meta=(UnsafeDuringActorConstruction="true"))
 	void SetPhysicsMaxAngularVelocityInRadians(float NewMaxAngVel, bool bAddToCurrent = false, FName BoneName = NAME_None);
 
-	/** 
-	 *	Get the angular velocity of a single body, in degrees per second. 
-	 *	@param BoneName			If a SkeletalMeshComponent, name of body to get velocity of. 'None' indicates root body.
-	 */
-	UE_DEPRECATED(4.18, "Use GetPhysicsAngularVelocityInDegrees instead.")
-	UFUNCTION(BlueprintCallable, Category="Physics", meta=(UnsafeDuringActorConstruction="true", DeprecatedFunction, DeprecationMessage="Use GetPhysicsAngularVelocityInDegrees instead"))	
-	FVector GetPhysicsAngularVelocity(FName BoneName = NAME_None) const
-	{
-		return GetPhysicsAngularVelocityInDegrees(BoneName);
-	}
-
 	/** 
 	 *	Get the angular velocity of a single body, in degrees per second. 
 	 *	@param BoneName			If a SkeletalMeshComponent, name of body to get velocity of. 'None' indicates root body.
diff --git a/Engine/Source/Runtime/Engine/Classes/Components/SceneComponent.h b/Engine/Source/Runtime/Engine/Classes/Components/SceneComponent.h
--- a/Engine/Source/Runtime/Engine/Classes/Components/SceneComponent.h
+++ b/Engine/Source/Runtime/Engine/Classes/Components/SceneComponent.h
@@ -673,14 +673,6 @@ public:
 	/** Backwards compatibility: Used to convert old-style EAttachLocation to new-style EAttachmentRules */
 	static void ConvertAttachLocation(EAttachLocation::Type InAttachLocation, EAttachmentRule& InOutLocationRule, EAttachmentRule& InOutRotationRule, EAttachmentRule& InOutScaleRule);
 
-	UE_DEPRECATED(4.12, "This function is deprecated, please use AttachToComponent instead.")
-	bool AttachTo(USceneComponent* InParent, FName InSocketName = NAME_None, EAttachLocation::Type AttachType = EAttachLocation::KeepRelativeOffset, bool bWeldSimulatedBodies = false);
-
-	/** DEPRECATED - Use AttachToComponent() instead */
-	UE_DEPRECATED(4.17, "This function is deprecated, please use AttachToComponent() instead.")
-	UFUNCTION(BlueprintCallable, Category = "Utilities|Transformation", meta = (DeprecationMessage = "OVERRIDE BAD MESSAGE", DisplayName = "AttachTo (Deprecated)", AttachType = "KeepRelativeOffset"))
-	bool K2_AttachTo(USceneComponent* InParent, FName InSocketName = NAME_None, EAttachLocation::Type AttachType = EAttachLocation::KeepRelativeOffset, bool bWeldSimulatedBodies = true);
-
 	/**
 	* Attach this component to another scene component, optionally at a named socket. It is valid to call this on components whether or not they have been Registered, however from
 	* constructor or when not registered it is preferable to use SetupAttachment.
@@ -704,16 +696,6 @@ public:
 	UFUNCTION(BlueprintCallable, Category = "Utilities|Transformation", meta = (DisplayName = "AttachComponentToComponent", ScriptName = "AttachToComponent", bWeldSimulatedBodies=true))
 	bool K2_AttachToComponent(USceneComponent* Parent, FName SocketName, EAttachmentRule LocationRule, EAttachmentRule RotationRule, EAttachmentRule ScaleRule, bool bWeldSimulatedBodies);
 
-	/** DEPRECATED - Use AttachToComponent() instead */
-	UE_DEPRECATED(4.17, "Use AttachToComponent() instead.")
-	UFUNCTION(BlueprintCallable, meta=(DeprecatedFunction, DeprecationMessage = "Use AttachToComponent instead."), Category="Utilities|Transformation")
-	bool SnapTo(USceneComponent* InParent, FName InSocketName = NAME_None);
-
-	/** DEPRECATED - Use DetachFromComponent() instead */
-	UE_DEPRECATED(4.12, "This function is deprecated, please use DetachFromComponent() instead.")
-	UFUNCTION(BlueprintCallable, Category="Utilities|Transformation", meta = (DisplayName = "DetachFromParent (Deprecated)"))
-	virtual void DetachFromParent(bool bMaintainWorldPosition = false, bool bCallModify = true);
-
 	/** 
 	 * Detach this component from whatever it is attached to. Automatically unwelds components that are welded together (See WeldTo)
 	 * @param LocationRule				How to handle translations when detaching.
diff --git a/Engine/Source/Runtime/Engine/Classes/Components/SplineMeshComponent.h b/Engine/Source/Runtime/Engine/Classes/Components/SplineMeshComponent.h
--- a/Engine/Source/Runtime/Engine/Classes/Components/SplineMeshComponent.h
+++ b/Engine/Source/Runtime/Engine/Classes/Components/SplineMeshComponent.h
@@ -205,9 +205,6 @@ public:
 	/** Called to notify render thread and possibly collision of a change in spline params or mesh */
 	void UpdateRenderStateAndCollision();
 
-	UE_DEPRECATED(4.11, "This method has been renamed to UpdateRenderStateAndCollision, but use of UpdateMesh is preferred")
-	void MarkSplineParamsDirty() { UpdateRenderStateAndCollision(); }
-
 	/** Update the collision and render state on the spline mesh following changes to its geometry */
 	UFUNCTION(BlueprintCallable, Category = SplineMesh)
 	void UpdateMesh();
diff --git a/Engine/Source/Runtime/Engine/Classes/Components/TextRenderComponent.h b/Engine/Source/Runtime/Engine/Classes/Components/TextRenderComponent.h
--- a/Engine/Source/Runtime/Engine/Classes/Components/TextRenderComponent.h
+++ b/Engine/Source/Runtime/Engine/Classes/Components/TextRenderComponent.h
@@ -99,14 +99,6 @@ class ENGINE_API UTextRenderComponent : public UPrimitiveComponent
 
 	// -----------------------------
 	
-	/**
-	 * Change the text value and signal the primitives to be rebuilt 
-	 * The FString variant is deprecated in favor of the FText variant
-	 */
-	UE_DEPRECATED(4.8, "Passing text as FString is deprecated, please use FText instead (likely via a LOCTEXT).")
-	UFUNCTION(BlueprintCallable, Category="Rendering|Components|TextRender", meta=(DisplayName="Set Text (String)", ScriptNoExport, DeprecatedFunction, DeprecationMessage="Use the SetText function taking an FText instead."))
-	void SetText(const FString& Value);
-
 	/** Change the text value and signal the primitives to be rebuilt */
 	void SetText(const FText& Value);
 
diff --git a/Engine/Source/Runtime/Engine/Classes/Engine/Engine.h b/Engine/Source/Runtime/Engine/Classes/Engine/Engine.h
--- a/Engine/Source/Runtime/Engine/Classes/Engine/Engine.h
+++ b/Engine/Source/Runtime/Engine/Classes/Engine/Engine.h
@@ -2462,16 +2462,6 @@ public:
 		return GetWorldFromContextObject(Object, EGetWorldErrorMode::Assert);
 	}
 
-	/** 
-	 * This function is deprecated
-	 */
-	UE_DEPRECATED(4.17, "GetWorldFromContextObject(Object) and GetWorldFromContextObject(Object, boolean) are replaced by GetWorldFromContextObject(Object, Enum) or GetWorldFromContextObjectChecked(Object)")
-	UWorld* GetWorldFromContextObject(const UObject* Object, bool bChecked = true) const
-	{
-		// Note: The behavior in 4.16 and before was similar to Assert if bChecked was true, but almost no callers actually wanted to pass in bChecked=true
-		return GetWorldFromContextObject(Object, bChecked ? EGetWorldErrorMode::LogAndReturnNull : EGetWorldErrorMode::ReturnNull);
-	}
-
 	/** 
 	 * mostly done to check if PIE is being set up, go GWorld is going to change, and it's not really _the_G_World_
 	 * NOTE: hope this goes away once PIE and regular game triggering are not that separate code paths
@@ -2864,9 +2854,6 @@ public:
 
 #if WITH_SERVER_CODE
 	virtual bool HandleServerTravelCommand( const TCHAR* Cmd, FOutputDevice& Ar, UWorld* InWorld );
-
-	UE_DEPRECATED(4.14, "Say Command moved to GameMode as an exec function")
-	virtual bool HandleSayCommand( const TCHAR* Cmd, FOutputDevice& Ar, UWorld* InWorld );
 #endif
 
 	virtual bool HandleDisconnectCommand( const TCHAR* Cmd, FOutputDevice& Ar, UWorld *InWorld );
diff --git a/Engine/Source/Runtime/Engine/Classes/Engine/EngineTypes.h b/Engine/Source/Runtime/Engine/Classes/Engine/EngineTypes.h
--- a/Engine/Source/Runtime/Engine/Classes/Engine/EngineTypes.h
+++ b/Engine/Source/Runtime/Engine/Classes/Engine/EngineTypes.h
@@ -893,9 +893,6 @@ namespace EWorldType
 		/** An editor world that was loaded but not currently being edited in the level editor */
 		Inactive
 	};
-
-	UE_DEPRECATED(4.14, "EWorldType::Preview is deprecated. Please use either EWorldType::EditorPreview or EWorldType::GamePreview")
-	const EWorldType::Type Preview = EWorldType::EditorPreview;
 }
 
 /** Describes what parts of level streaming should be forcibly handled immediately */
diff --git a/Engine/Source/Runtime/Engine/Classes/Engine/GameInstance.h b/Engine/Source/Runtime/Engine/Classes/Engine/GameInstance.h
--- a/Engine/Source/Runtime/Engine/Classes/Engine/GameInstance.h
+++ b/Engine/Source/Runtime/Engine/Classes/Engine/GameInstance.h
@@ -248,12 +248,6 @@ public:
 
 	/** Called as soon as the game mode is spawned, to allow additional PIE setting validation prior to creating the local players / etc... (called on pure clients too, in which case the game mode is nullptr) */
 	virtual FGameInstancePIEResult PostCreateGameModeForPIE(const FGameInstancePIEParameters& Params, AGameModeBase* GameMode);
-
-	UE_DEPRECATED(4.15, "Please override InitializeForPIE instead")
-	virtual bool InitializePIE(bool bAnyBlueprintErrors, int32 PIEInstance, bool bRunAsDedicated);
-
-	UE_DEPRECATED(4.15, "Please override StartPlayInEditorGameInstance instead")
-	virtual bool StartPIEGameInstance(ULocalPlayer* LocalPlayer, bool bInSimulateInEditor, bool bAnyBlueprintErrors, bool bStartInSpectatorMode);
 #endif
 
 	class UEngine* GetEngine() const;
diff --git a/Engine/Source/Runtime/Engine/Classes/Engine/Level.h b/Engine/Source/Runtime/Engine/Classes/Engine/Level.h
--- a/Engine/Source/Runtime/Engine/Classes/Engine/Level.h
+++ b/Engine/Source/Runtime/Engine/Classes/Engine/Level.h
@@ -641,10 +641,6 @@ public:
 	int32										CurrentActorIndexForUnregisterComponents;
 
 
-	/** Whether the level is currently pending being made visible.							*/
-	UE_DEPRECATED(4.15, "Use HasVisibilityChangeRequestPending")
-	bool HasVisibilityRequestPending() const;
-
 	/** Whether the level is currently pending being made invisible or visible.				*/
 	ENGINE_API bool HasVisibilityChangeRequestPending() const;
 
diff --git a/Engine/Source/Runtime/Engine/Classes/Engine/World.h b/Engine/Source/Runtime/Engine/Classes/Engine/World.h
--- a/Engine/Source/Runtime/Engine/Classes/Engine/World.h
+++ b/Engine/Source/Runtime/Engine/Classes/Engine/World.h
@@ -2713,33 +2713,6 @@ public:
 	 */
 	void MarkObjectsPendingKill();
 
-	/**
-	 *  Interface to allow WorldSettings to request immediate garbage collection
-	 */
-	UE_DEPRECATED(4.18, "Use GEngine->PerformGarbageCollectionAndCleanupActors instead.")
-	void PerformGarbageCollectionAndCleanupActors();
-
-	/**
-	 *  Requests a one frame delay of Garbage Collection
-	 */
-	UE_DEPRECATED(4.18, "Use GEngine->DelayGarbageCollection instead.")
-	void DelayGarbageCollection();
-
-	/**
-	 * Updates the timer (as a one-off) that is used to trigger garbage collection; this should only be used for things
-	 * like performance tests, using it recklessly can dramatically increase memory usage and cost of the eventual GC.
-	 *
-	 * Note: Things that force a GC will still force a GC after using this method (and they will also reset the timer)
-	 */
-	UE_DEPRECATED(4.18, "Use GEngine->SetTimeUntilNextGarbageCollection instead.")
-	void SetTimeUntilNextGarbageCollection(float MinTimeUntilNextPass);
-
-	/**
-	 * Returns the current desired time between garbage collection passes (not the time remaining)
-	 */
-	UE_DEPRECATED(4.18, "Call GEngine->GetTimeBetweenGarbageCollectionPasses instead")
-	float GetTimeBetweenGarbageCollectionPasses() const;
-
 	/**
 	 *	Remove NULL entries from actor list. Only does so for dynamic actors to avoid resorting. 
 	 *	In theory static actors shouldn't be deleted during gameplay.
@@ -3524,10 +3497,6 @@ public:
 	/** @return the current detail mode, like EDetailMode but can be outside of the range */
 	int32 GetDetailMode();
 
-	/** Updates the timer between garbage collection such that at the next opportunity garbage collection will be run. */
-	UE_DEPRECATED(4.18, "Call GEngine->ForceGarbageCollection instead")
-	void ForceGarbageCollection( bool bFullPurge = false );
-
 	/** asynchronously loads the given levels in preparation for a streaming map transition.
 	 * This codepath is designed for worlds that heavily use level streaming and GameModes where the game state should
 	 * be preserved through a transition.
diff --git a/Engine/Source/Runtime/Engine/Classes/GameFramework/Actor.h b/Engine/Source/Runtime/Engine/Classes/GameFramework/Actor.h
--- a/Engine/Source/Runtime/Engine/Classes/GameFramework/Actor.h
+++ b/Engine/Source/Runtime/Engine/Classes/GameFramework/Actor.h
@@ -1395,15 +1395,6 @@ public:
 	UFUNCTION(BlueprintCallable, meta=(BlueprintInternalUseOnly="true"))
 	void FinishAddComponent(UActorComponent* Component, bool bManualAttachment, const FTransform& RelativeTransform);
 
-	UE_DEPRECATED(4.17, "Use UActorComponent::DestroyComponent() instead")
-	UFUNCTION(BlueprintCallable, meta=(DeprecatedFunction, DeprecationMessage = "Use Component.DestroyComponent instead", BlueprintProtected = "true", DisplayName = "DestroyComponent", ScriptName = "DestroyComponent"))
-	void K2_DestroyComponent(UActorComponent* Component);
-
-	/** DEPRECATED - Use AttachToComponent() instead */
-	UE_DEPRECATED(4.17, "Use AttachToComponent() instead.")
-	UFUNCTION(BlueprintCallable, meta = (DisplayName = "AttachRootComponentTo (Deprecated)", ScriptNoExport, AttachLocationType = "KeepRelativeOffset"), Category = "Utilities|Transformation")
-	void K2_AttachRootComponentTo(USceneComponent* InParent, FName InSocketName = NAME_None, EAttachLocation::Type AttachLocationType = EAttachLocation::KeepRelativeOffset, bool bWeldSimulatedBodies = true);
-
 	/**
 	 * Attaches the RootComponent of this Actor to the supplied component, optionally at a named socket. It is not valid to call this on components that are not Registered.
 	 * @param Parent					Parent to attach to.
@@ -1424,11 +1415,6 @@ public:
 	 */
 	void AttachToComponent(USceneComponent* Parent, const FAttachmentTransformRules& AttachmentRules, FName SocketName = NAME_None);
 
-	/** DEPRECATED - Use AttachToActor() instead */
-	UE_DEPRECATED(4.17, "Use AttachToActor() instead.")
-	UFUNCTION(BlueprintCallable, meta = (DisplayName = "AttachRootComponentToActor (Deprecated)", ScriptNoExport, AttachLocationType = "KeepRelativeOffset"), Category = "Utilities|Transformation")
-	void K2_AttachRootComponentToActor(AActor* InParentActor, FName InSocketName = NAME_None, EAttachLocation::Type AttachLocationType = EAttachLocation::KeepRelativeOffset, bool bWeldSimulatedBodies = true);
-
 	/**
 	 * Attaches the RootComponent of this Actor to the RootComponent of the supplied actor, optionally at a named socket.
 	 * @param ParentActor				Actor to attach this actor's RootComponent to
@@ -1449,15 +1435,6 @@ public:
 	UFUNCTION(BlueprintCallable, meta = (DisplayName = "AttachActorToActor", ScriptName = "AttachToActor", bWeldSimulatedBodies=true), Category = "Utilities|Transformation")
 	void K2_AttachToActor(AActor* ParentActor, FName SocketName, EAttachmentRule LocationRule, EAttachmentRule RotationRule, EAttachmentRule ScaleRule, bool bWeldSimulatedBodies);
 
-	UE_DEPRECATED(4.17, "Use AttachToComponent() with EAttachLocation::SnapToTarget option instead")
-	UFUNCTION(BlueprintCallable, meta=(DeprecatedFunction, DeprecationMessage = "Use AttachRootComponentTo with EAttachLocation::SnapToTarget option instead", DisplayName = "SnapActorTo", ScriptNoExport), Category="Utilities|Transformation")
-	void SnapRootComponentTo(AActor* InParentActor, FName InSocketName);
-
-	/** DEPRECATED - Use DetachFromActor() instead */
-	UE_DEPRECATED(4.17, "Use DetachFromActor() instead")
-	UFUNCTION(BlueprintCallable, meta=(DisplayName = "DetachActorFromActor (Deprecated)", ScriptNoExport), Category="Utilities|Transformation")
-	void DetachRootComponentFromParent(bool bMaintainWorldPosition = true);
-
 	/** 
 	 * Detaches the RootComponent of this Actor from any SceneComponent it is currently attached to. 
 	 * @param  LocationRule				How to handle translation when detaching.
@@ -1524,10 +1501,6 @@ public:
 	UFUNCTION(BlueprintCallable, Category="Utilities")
 	void SetTickableWhenPaused(bool bTickableWhenPaused);
 
-	UE_DEPRECATED(4.17, "Use UPrimitiveComponent::CreateAndSetMaterialInstanceDynamic() instead.")
-	UFUNCTION(BlueprintCallable, meta=(DeprecatedFunction, DeprecationMessage="Use PrimitiveComponent.CreateAndSetMaterialInstanceDynamic instead.", BlueprintProtected = "true"), Category="Rendering|Material")
-	class UMaterialInstanceDynamic* MakeMIDForMaterial(class UMaterialInterface* Parent);
-
 	/** The number of seconds (in game time) since this Actor was created, relative to Get Game Time In Seconds. */
 	UFUNCTION(BlueprintPure, Category=Actor)
 	float GetGameTimeSinceCreation() const;
diff --git a/Engine/Source/Runtime/Engine/Classes/GameFramework/CharacterMovementComponent.h b/Engine/Source/Runtime/Engine/Classes/GameFramework/CharacterMovementComponent.h
--- a/Engine/Source/Runtime/Engine/Classes/GameFramework/CharacterMovementComponent.h
+++ b/Engine/Source/Runtime/Engine/Classes/GameFramework/CharacterMovementComponent.h
@@ -1425,13 +1425,6 @@ public:
 	UFUNCTION(BlueprintCallable, Category = "Pawn|Components|CharacterMovement")
 	virtual float GetMinAnalogSpeed() const;
 	
-	UE_DEPRECATED(4.3, "GetModifiedMaxAcceleration() is deprecated, apply your own modifiers to GetMaxAcceleration() if desired.")
-	virtual float GetModifiedMaxAcceleration() const;
-	
-	/** Returns maximum acceleration for the current state, based on MaxAcceleration and any additional modifiers. */
-	UFUNCTION(BlueprintCallable, Category="Pawn|Components|CharacterMovement", meta=(DeprecatedFunction, DisplayName="GetModifiedMaxAcceleration", ScriptName="GetModifiedMaxAcceleration", DeprecationMessage="GetModifiedMaxAcceleration() is deprecated, apply your own modifiers to GetMaxAcceleration() if desired."))
-	virtual float K2_GetModifiedMaxAcceleration() const;
-
 	/** Returns maximum acceleration for the current state. */
 	UFUNCTION(BlueprintCallable, Category="Pawn|Components|CharacterMovement")
 	virtual float GetMaxAcceleration() const;
@@ -2541,9 +2534,6 @@ public:
 	 */
 	virtual FVector CalcAnimRootMotionVelocity(const FVector& RootMotionDeltaMove, float DeltaSeconds, const FVector& CurrentVelocity) const;
 
-	UE_DEPRECATED(4.13, "CalcRootMotionVelocity() has been replaced by CalcAnimRootMotionVelocity() instead, and ConstrainAnimRootMotionVelocity() now handles restricting root motion velocity under different conditions.")
-	virtual FVector CalcRootMotionVelocity(const FVector& RootMotionDeltaMove, float DeltaSeconds, const FVector& CurrentVelocity) const;
-
 	/**
 	 * Constrain components of root motion velocity that may not be appropriate given the current movement mode (e.g. when falling Z may be ignored).
 	 */
@@ -2844,10 +2834,7 @@ public:
 
 	// Mesh smoothing variables (for network smoothing)
 	//
-	/** Whether to smoothly interpolate pawn position corrections on clients based on received location updates */
-	UE_DEPRECATED(4.11, "bSmoothNetUpdates will be removed, use UCharacterMovementComponent::NetworkSmoothingMode instead.")
-	uint32 bSmoothNetUpdates:1;
-
+	
 	/** Used for position smoothing in net games */
 	FVector OriginalMeshTranslationOffset;
 
@@ -2878,14 +2865,6 @@ public:
 	/** Used to track the client time as we try to match the server.*/
 	double SmoothingClientTimeStamp;
 
-	/** Used to track how much time has elapsed since last correction. It can be computed as World->TimeSince(LastCorrectionTime). */
-	UE_DEPRECATED(4.11, "CurrentSmoothTime will be removed, use LastCorrectionTime instead.")
-	float CurrentSmoothTime;
-
-	/** Used to signify that linear smoothing is desired */
-	UE_DEPRECATED(4.11, "bUseLinearSmoothing will be removed, use UCharacterMovementComponent::NetworkSmoothingMode instead.")
-	bool bUseLinearSmoothing;
-
 	/**
 	 * Copied value from UCharacterMovementComponent::NetworkMaxSmoothUpdateDistance.
 	 * @see UCharacterMovementComponent::NetworkMaxSmoothUpdateDistance
@@ -2903,10 +2882,6 @@ public:
 
 	/** How long to take to smoothly interpolate from the old pawn rotation on the client to the corrected one sent by the server.  Must be >= 0. Not used for linear smoothing. */
 	float SmoothNetUpdateRotationTime;
-
-	/** (DEPRECATED) How long server will wait for client move update before setting position */
-	UE_DEPRECATED(4.12, "MaxResponseTime has been renamed to MaxMoveDeltaTime for clarity in what it does and will be removed, use MaxMoveDeltaTime instead.")
-	float MaxResponseTime;
 	
 	/** 
 	 * Max delta time for a given move, in real seconds
@@ -2976,10 +2951,6 @@ public:
 
 	/** Server clock time when last server move was received from client (does NOT include forced moves on server) */
 	float ServerTimeStampLastServerMove;
-
-	/** (DEPRECATED) How long server will wait for client move update before setting position */
-	UE_DEPRECATED(4.12, "MaxResponseTime has been renamed to MaxMoveDeltaTime for clarity in what it does and will be removed, use MaxMoveDeltaTime instead.")
-	float MaxResponseTime;
 	
 	/** 
 	 * Max delta time for a given move, in real seconds
diff --git a/Engine/Source/Runtime/Engine/Classes/GameFramework/Controller.h b/Engine/Source/Runtime/Engine/Classes/GameFramework/Controller.h
--- a/Engine/Source/Runtime/Engine/Classes/GameFramework/Controller.h
+++ b/Engine/Source/Runtime/Engine/Classes/GameFramework/Controller.h
@@ -188,11 +188,6 @@ public:
 		return Cast<T>(PlayerState);
 	}
 
-	/** DEPRECATED! Use the standard "Cast To" node instead. Casts this Controller to a Player Controller, if possible. */
-	UE_DEPRECATED(4.11, "CastToPlayerController has been replaced by the standard Cast To node.")
-	UFUNCTION(BlueprintCallable, Category=Pawn, meta=(DeprecatedFunction, DeprecationMessage="Use standard Cast To node instead."))
-	class APlayerController* CastToPlayerController();
-
 	/** Replicated function to set the pawn location and rotation, allowing server to force (ex. teleports). */
 	UFUNCTION(Reliable, Client)
 	void ClientSetLocation(FVector NewLocation, FRotator NewRotation);
diff --git a/Engine/Source/Runtime/Engine/Classes/GameFramework/GameMode.h b/Engine/Source/Runtime/Engine/Classes/GameFramework/GameMode.h
--- a/Engine/Source/Runtime/Engine/Classes/GameFramework/GameMode.h
+++ b/Engine/Source/Runtime/Engine/Classes/GameFramework/GameMode.h
@@ -162,33 +162,13 @@ public:
 	/** Removes the passed in player controller from the correct count for player/spectator/tranistioning **/
 	void RemovePlayerControllerFromPlayerCount(APlayerController* PC);
 
-	UE_DEPRECATED(4.14, "Deprecated in favor of PreloadContentForURL on GameInstance")
-	virtual FString GetDefaultGameClassPath(const FString& MapName, const FString& Options, const FString& Portal) const;
-
-	UE_DEPRECATED(4.14, "Deprecated in favor of OverrideGameModeClass on GameInstance")
-	virtual TSubclassOf<AGameMode> GetGameModeClass(const FString& MapName, const FString& Options, const FString& Portal) const;
-
-	UE_DEPRECATED(4.14, "Deprecated in favor of GetGameModeForName on GameMapsSettings")
-	static FString StaticGetFullGameClassName(FString const& Str);
-
 	/** Return true if we want to travel_absolute, used by RestartGame by default */
 	virtual bool GetTravelType();
 
-	UE_DEPRECATED(4.14, "SendPlayer is not in use, call ClientTravel directly instead")
-	virtual void SendPlayer( APlayerController* aPlayer, const FString& URL );
-	
-	UE_DEPRECATED(4.14, "StartNewPlayer has been split into InitializeHUDForPlayer and HandleStartingNewPlayer")
-	virtual void StartNewPlayer(APlayerController* NewPlayer);
-
 	/** Exec command to broadcast a string to all players */
 	UFUNCTION(Exec, BlueprintCallable, Category = AI)
 	virtual void Say(const FString& Msg);
 
-	/** Alters the synthetic bandwidth limit for a running game. */
-	UE_DEPRECATED(4.17, "AsyncIOBandwidthLimit is no longer configurable")
-	UFUNCTION(exec)	
-	virtual void SetBandwidthLimit(float AsyncIOBandwidthLimit) {}
-
 	/** Broadcast a string to all players. */
 	virtual void Broadcast( AActor* Sender, const FString& Msg, FName Type = NAME_None );
 
diff --git a/Engine/Source/Runtime/Engine/Classes/GameFramework/HUD.h b/Engine/Source/Runtime/Engine/Classes/GameFramework/HUD.h
--- a/Engine/Source/Runtime/Engine/Classes/GameFramework/HUD.h
+++ b/Engine/Source/Runtime/Engine/Classes/GameFramework/HUD.h
@@ -37,16 +37,6 @@ class ENGINE_API AHUD : public AActor
 {
 	GENERATED_UCLASS_BODY()
 
-	/** Pre-defined FColors for convenience. */
-	UE_DEPRECATED(4.13, "Use FColor::White instead")
-	static const FColor WhiteColor;
-
-	UE_DEPRECATED(4.13, "Use FColor::Green instead")
-	static const FColor GreenColor;
-
-	UE_DEPRECATED(4.13, "Use FColor::Red instead")
-	static const FColor RedColor;
-
 	/** PlayerController which owns this HUD. */
 	UPROPERTY(BlueprintReadOnly, Category=HUD)
 	APlayerController* PlayerOwner;    
diff --git a/Engine/Source/Runtime/Engine/Classes/GameFramework/MovementComponent.h b/Engine/Source/Runtime/Engine/Classes/GameFramework/MovementComponent.h
--- a/Engine/Source/Runtime/Engine/Classes/GameFramework/MovementComponent.h
+++ b/Engine/Source/Runtime/Engine/Classes/GameFramework/MovementComponent.h
@@ -210,20 +210,6 @@ public:
 	UFUNCTION(BlueprintCallable, Category="Components|Movement")
 	virtual float GetMaxSpeed() const;
 
-	UE_DEPRECATED(4.3, "GetMaxSpeedModifier() is deprecated, apply your own modifiers to GetMaxSpeed() if desired.")
-	virtual float GetMaxSpeedModifier() const;
-	
-	/** Returns a scalar applied to the maximum velocity that the component can currently move. */
-	UFUNCTION(BlueprintCallable, Category="Components|Movement", meta=(DeprecatedFunction, DisplayName="GetMaxSpeedModifier", ScriptName="GetMaxSpeedModifier", DeprecationMessage="GetMaxSpeedModifier() is deprecated, apply your own modifiers to GetMaxSpeed() if desired."))
-	virtual float K2_GetMaxSpeedModifier() const;
-
-	UE_DEPRECATED(4.3, "GetModifiedMaxSpeed() is deprecated, use GetMaxSpeed() instead.")
-	virtual float GetModifiedMaxSpeed() const;
-
-	/** Returns the result of GetMaxSpeed() * GetMaxSpeedModifier(). */
-	UFUNCTION(BlueprintCallable, Category="Components|Movement", meta=(DisplayName="GetModifiedMaxSpeed", ScriptName="GetModifiedMaxSpeed", DeprecationMessage="GetModifiedMaxSpeed() is deprecated, apply your own modifiers to GetMaxSpeed() if desired."))
-	virtual float K2_GetModifiedMaxSpeed() const;
-
 	/**
 	 * Returns true if the current velocity is exceeding the given max speed (usually the result of GetMaxSpeed()), within a small error tolerance.
 	 * Note that under normal circumstances updates cause by acceleration will not cause this to be true, however external forces or changes in the max speed limit
diff --git a/Engine/Source/Runtime/Engine/Classes/GameFramework/Pawn.h b/Engine/Source/Runtime/Engine/Classes/GameFramework/Pawn.h
--- a/Engine/Source/Runtime/Engine/Classes/GameFramework/Pawn.h
+++ b/Engine/Source/Runtime/Engine/Classes/GameFramework/Pawn.h
@@ -481,16 +481,6 @@ public:
 	/** Remove an Actor to ignore by Pawn's movement collision */
 	void MoveIgnoreActorRemove(AActor* ActorToIgnore);
 
-	// DEPRECATED FUNCTIONS
-
-	/** (Deprecated) Launch Character with LaunchVelocity  */
-	UE_DEPRECATED(4.8, "LaunchPawn is deprecated. For Characters, use LaunchCharacter() instead.")
-	UFUNCTION(BlueprintCallable, Category="Pawn", meta=(DeprecatedFunction, DeprecationMessage="Use Character.LaunchCharacter instead"))
-	void LaunchPawn(FVector LaunchVelocity, bool bXYOverride, bool bZOverride);
-
-	/** (Deprecated) Return the input vector in world space. */
-	UFUNCTION(BlueprintCallable, Category="Pawn|Input", meta=(DeprecatedFunction, DisplayName="GetMovementInputVector", ScriptName="GetMovementInputVector", DeprecationMessage="GetMovementInputVector has been deprecated, use either GetPendingMovementInputVector or GetLastMovementInputVector"))
-	FVector K2_GetMovementInputVector() const;	
 };
 
 
diff --git a/Engine/Source/Runtime/Engine/Classes/GameFramework/PawnMovementComponent.h b/Engine/Source/Runtime/Engine/Classes/GameFramework/PawnMovementComponent.h
--- a/Engine/Source/Runtime/Engine/Classes/GameFramework/PawnMovementComponent.h
+++ b/Engine/Source/Runtime/Engine/Classes/GameFramework/PawnMovementComponent.h
@@ -97,11 +97,4 @@ protected:
 public:
 
 	virtual void Serialize(FArchive& Ar) override;
-
-	// DEPRECATED FUNCTIONS
-
-	/** (Deprecated) Return the input vector in world space. */
-	UE_DEPRECATED(4.5, "GetInputVector() has been deprecated, use either GetPendingInputVector() or GetLastInputVector().")
-	UFUNCTION(BlueprintCallable, Category="Pawn|Components|PawnMovement", meta=(DeprecatedFunction, DisplayName="GetInputVector", ScriptName="GetInputVector", DeprecationMessage="GetInputVector has been deprecated, use either GetPendingInputVector or GetLastInputVector"))
-	FVector K2_GetInputVector() const;
 };
diff --git a/Engine/Source/Runtime/Engine/Classes/Kismet/GameplayStatics.h b/Engine/Source/Runtime/Engine/Classes/Kismet/GameplayStatics.h
--- a/Engine/Source/Runtime/Engine/Classes/Kismet/GameplayStatics.h
+++ b/Engine/Source/Runtime/Engine/Classes/Kismet/GameplayStatics.h
@@ -58,10 +58,6 @@ class ENGINE_API UGameplayStatics : public UBlueprintFunctionLibrary
 	UFUNCTION(BlueprintCallable, Category="Spawning", meta=(WorldContext="WorldContextObject", UnsafeDuringActorConstruction = "true", BlueprintInternalUseOnly = "true", DeprecatedFunction, DeprecationMessage="Use BeginSpawningActorFromClass"))
 	static class AActor* BeginSpawningActorFromBlueprint(const UObject* WorldContextObject, const class UBlueprint* Blueprint, const FTransform& SpawnTransform, bool bNoCollisionFail);
 
-	UE_DEPRECATED(4.9, "This function is deprecated. Please use BeginDeferredActorSpawnFromClass instead.")
-	UFUNCTION(BlueprintCallable, Category="Spawning", meta=(WorldContext="WorldContextObject", UnsafeDuringActorConstruction = "true", BlueprintInternalUseOnly = "true"))
-	static class AActor* BeginSpawningActorFromClass(const UObject* WorldContextObject, TSubclassOf<AActor> ActorClass, const FTransform& SpawnTransform, bool bNoCollisionFail = false, AActor* Owner = nullptr);
-
 	/** Spawns an instance of an actor class, but does not automatically run its construction script.  */
 	UFUNCTION(BlueprintCallable, Category = "Spawning", meta = (WorldContext = "WorldContextObject", UnsafeDuringActorConstruction = "true", BlueprintInternalUseOnly = "true"))
 	static class AActor* BeginDeferredActorSpawnFromClass(const UObject* WorldContextObject, TSubclassOf<AActor> ActorClass, const FTransform& SpawnTransform, ESpawnActorCollisionHandlingMethod CollisionHandlingOverride = ESpawnActorCollisionHandlingMethod::Undefined, AActor* Owner = nullptr);
@@ -1122,12 +1118,6 @@ public:
 	*/
 	static bool PredictProjectilePath(const UObject* WorldContextObject, const FPredictProjectilePathParams& PredictParams, FPredictProjectilePathResult& PredictResult);
 
-	/**
-	* Deprecated version, use version with input/output struct params instead.
-	*/
-	UE_DEPRECATED(4.15, "PredictProjectilePath with many parameters has been deprecated in favor of the version taking single input parameter and output result structs.")
-	static bool PredictProjectilePath(const UObject* WorldContextObject, FHitResult& OutHit, TArray<FVector>& OutPathPositions, FVector& OutLastTraceDestination, FVector StartPos, FVector LaunchVelocity, bool bTracePath, float ProjectileRadius, const TArray<TEnumAsByte<EObjectTypeQuery> >& ObjectTypes, bool bTraceComplex, const TArray<AActor*>& ActorsToIgnore, EDrawDebugTrace::Type DrawDebugType, float DrawDebugTime, float SimFrequency = 15.f, float MaxSimTime = 2.f, float OverrideGravityZ = 0);
-
 	/**
 	* Predict the arc of a virtual projectile affected by gravity with collision checks along the arc.
 	* Returns true if it hit something.
diff --git a/Engine/Source/Runtime/Engine/Classes/Kismet/HeadMountedDisplayFunctionLibrary.h b/Engine/Source/Runtime/Engine/Classes/Kismet/HeadMountedDisplayFunctionLibrary.h
deleted file mode 100644
diff --git a/Engine/Source/Runtime/Engine/Private/Actor.cpp b/Engine/Source/Runtime/Engine/Private/Actor.cpp
--- a/Engine/Source/Runtime/Engine/Private/Actor.cpp
+++ b/Engine/Source/Runtime/Engine/Private/Actor.cpp
@@ -1914,16 +1914,6 @@ void AActor::SetAutoDestroyWhenFinished(bool bVal)
 	}
 }
 
-void AActor::K2_AttachRootComponentTo(USceneComponent* InParent, FName InSocketName, EAttachLocation::Type AttachLocationType /*= EAttachLocation::KeepRelativeOffset */, bool bWeldSimulatedBodies /*=true*/)
-{
-PRAGMA_DISABLE_DEPRECATION_WARNINGS
-	if (RootComponent && InParent)
-	{
-		RootComponent->AttachTo(InParent, InSocketName, AttachLocationType, bWeldSimulatedBodies);
-	}
-PRAGMA_ENABLE_DEPRECATION_WARNINGS
-}
-
 void AActor::K2_AttachToComponent(USceneComponent* Parent, FName SocketName, EAttachmentRule LocationRule, EAttachmentRule RotationRule, EAttachmentRule ScaleRule, bool bWeldSimulatedBodies)
 {
 	AttachToComponent(Parent, FAttachmentTransformRules(LocationRule, RotationRule, ScaleRule, bWeldSimulatedBodies), SocketName);
@@ -1982,20 +1972,6 @@ void AActor::OnRep_AttachmentReplication()
 	}
 }
 
-void AActor::K2_AttachRootComponentToActor(AActor* InParentActor, FName InSocketName /*= NAME_None*/, EAttachLocation::Type AttachLocationType /*= EAttachLocation::KeepRelativeOffset */, bool bWeldSimulatedBodies /*=true*/)
-{
-	if (RootComponent && InParentActor)
-	{
-		USceneComponent* ParentDefaultAttachComponent = InParentActor->GetDefaultAttachComponent();
-		if (ParentDefaultAttachComponent)
-		{
-PRAGMA_DISABLE_DEPRECATION_WARNINGS
-			RootComponent->AttachTo(ParentDefaultAttachComponent, InSocketName, AttachLocationType, bWeldSimulatedBodies );
-PRAGMA_ENABLE_DEPRECATION_WARNINGS
-		}
-	}
-}
-
 void AActor::K2_AttachToActor(AActor* ParentActor, FName SocketName, EAttachmentRule LocationRule, EAttachmentRule RotationRule, EAttachmentRule ScaleRule, bool bWeldSimulatedBodies)
 {
 	AttachToActor(ParentActor, FAttachmentTransformRules(LocationRule, RotationRule, ScaleRule, bWeldSimulatedBodies), SocketName);
@@ -2013,32 +1989,6 @@ void AActor::AttachToActor(AActor* ParentActor, const FAttachmentTransformRules&
 	}
 }
 
-void AActor::SnapRootComponentTo(AActor* InParentActor, FName InSocketName/* = NAME_None*/)
-{
-	if (RootComponent && InParentActor)
-	{
-		USceneComponent* ParentDefaultAttachComponent = InParentActor->GetDefaultAttachComponent();
-		if (ParentDefaultAttachComponent)
-		{
-			RootComponent->AttachToComponent(ParentDefaultAttachComponent, FAttachmentTransformRules::SnapToTargetNotIncludingScale, InSocketName);
-		}
-	}
-}
-
-void AActor::DetachRootComponentFromParent(bool bMaintainWorldPosition)
-{
-	if(RootComponent)
-	{
-PRAGMA_DISABLE_DEPRECATION_WARNINGS
-		RootComponent->DetachFromParent(bMaintainWorldPosition);
-PRAGMA_ENABLE_DEPRECATION_WARNINGS		
-
-		// Clear AttachmentReplication struct
-		AttachmentReplication = FRepAttachment();
-		MARK_PROPERTY_DIRTY_FROM_NAME(AActor, AttachmentReplication, this);
-	}
-}
-
 void AActor::K2_DetachFromActor(EDetachmentRule LocationRule /*= EDetachmentRule::KeepRelative*/, EDetachmentRule RotationRule /*= EDetachmentRule::KeepRelative*/, EDetachmentRule ScaleRule /*= EDetachmentRule::KeepRelative*/)
 {
 	DetachFromActor(FDetachmentTransformRules(LocationRule, RotationRule, ScaleRule, true));
@@ -4230,15 +4180,6 @@ void AActor::K2_DestroyActor()
 	Destroy();
 }
 
-void AActor::K2_DestroyComponent(UActorComponent* Component)
-{
-	// If its a valid component, and we own it, destroy it
-	if(Component && Component->GetOwner() == this)
-	{
-		Component->DestroyComponent();
-	}
-}
-
 bool AActor::SetRootComponent(class USceneComponent* NewRootComponent)
 {
 	/** Only components owned by this actor can be used as a its root component. */
@@ -5216,16 +5157,6 @@ float AActor::GetActorTimeDilation(const UWorld& ActorWorld) const
 	return CustomTimeDilation * ActorWorld.GetWorldSettings()->GetEffectiveTimeDilation();
 }
 
-UMaterialInstanceDynamic* AActor::MakeMIDForMaterial(class UMaterialInterface* Parent)
-{
-	// Deprecating this function. 
-	// Please use PrimitiveComponent->CreateAndSetMaterialInstanceDynamic
-	// OR PrimitiveComponent->CreateAndSetMaterialInstanceDynamicFromMaterial
-	// OR UMaterialInstanceDynamic::Create
-
-	return nullptr;
-}
-
 float AActor::GetDistanceTo(const AActor* OtherActor) const
 {
 	return OtherActor ? (GetActorLocation() - OtherActor->GetActorLocation()).Size() : 0.f;
diff --git a/Engine/Source/Runtime/Engine/Private/Components/CharacterMovementComponent.cpp b/Engine/Source/Runtime/Engine/Private/Components/CharacterMovementComponent.cpp
--- a/Engine/Source/Runtime/Engine/Private/Components/CharacterMovementComponent.cpp
+++ b/Engine/Source/Runtime/Engine/Private/Components/CharacterMovementComponent.cpp
@@ -1685,14 +1685,6 @@ void UCharacterMovementComponent::SimulateRootMotion(float DeltaSeconds, const F
 	RootMotionParams.Clear();
 }
 
-
-// TODO: Deprecated, remove.
-FVector UCharacterMovementComponent::CalcRootMotionVelocity(const FVector& RootMotionDeltaMove, float DeltaSeconds, const FVector& CurrentVelocity) const
-{
-	return CalcAnimRootMotionVelocity(RootMotionDeltaMove, DeltaSeconds, CurrentVelocity);
-}
-
-
 FVector UCharacterMovementComponent::CalcAnimRootMotionVelocity(const FVector& RootMotionDeltaMove, float DeltaSeconds, const FVector& CurrentVelocity) const
 {
 	if (ensure(DeltaSeconds > 0.f))
@@ -3672,25 +3664,6 @@ float UCharacterMovementComponent::GetMaxJumpHeightWithJumpTime() const
 	return MaxJumpHeight;
 }
 
-// TODO: deprecated, remove.
-float UCharacterMovementComponent::GetModifiedMaxAcceleration() const
-{
-	// Allow calling old deprecated function to maintain old behavior until it is removed.
-	PRAGMA_DISABLE_DEPRECATION_WARNINGS
-	return CharacterOwner ? MaxAcceleration * GetMaxSpeedModifier() : 0.f;
-	PRAGMA_ENABLE_DEPRECATION_WARNINGS
-}
-
-// TODO: deprecated, remove.
-float UCharacterMovementComponent::K2_GetModifiedMaxAcceleration() const
-{
-	// Allow calling old deprecated function to maintain old behavior until it is removed.
-	PRAGMA_DISABLE_DEPRECATION_WARNINGS
-	return GetModifiedMaxAcceleration();
-	PRAGMA_ENABLE_DEPRECATION_WARNINGS
-}
-
-
 float UCharacterMovementComponent::GetMaxAcceleration() const
 {
 	return MaxAcceleration;
@@ -10968,7 +10941,6 @@ FNetworkPredictionData_Client_Character::FNetworkPredictionData_Client_Character
 	, MaxFreeMoveCount(96)
 	, MaxSavedMoveCount(96)
 	, bUpdatePosition(false)
-	, bSmoothNetUpdates(false) // Deprecated
 	, OriginalMeshTranslationOffset(ForceInitToZero)
 	, MeshTranslationOffset(ForceInitToZero)
 	, OriginalMeshRotationOffset(FQuat::Identity)
@@ -10979,13 +10951,10 @@ FNetworkPredictionData_Client_Character::FNetworkPredictionData_Client_Character
 	, MaxClientSmoothingDeltaTime(0.5f)
 	, SmoothingServerTimeStamp(0.f)
 	, SmoothingClientTimeStamp(0.f)
-	, CurrentSmoothTime(0.f) // Deprecated
-	, bUseLinearSmoothing(false) // Deprecated
 	, MaxSmoothNetUpdateDist(0.f)
 	, NoSmoothNetUpdateDist(0.f)
 	, SmoothNetUpdateTime(0.f)
 	, SmoothNetUpdateRotationTime(0.f)
-	, MaxResponseTime(0.125f) // Deprecated, use MaxMoveDeltaTime instead
 	, MaxMoveDeltaTime(0.125f)
 	, LastSmoothLocation(FVector::ZeroVector)
 	, LastServerLocation(FVector::ZeroVector)
@@ -11006,8 +10975,6 @@ FNetworkPredictionData_Client_Character::FNetworkPredictionData_Client_Character
 		MaxClientSmoothingDeltaTime = FMath::Max(GameNetworkManager->MaxClientSmoothingDeltaTime, MaxMoveDeltaTime * 2.0f);
 	}
 
-	MaxResponseTime = MaxMoveDeltaTime; // MaxResponseTime is deprecated, use MaxMoveDeltaTime instead
-
 	if (ClientMovement.GetOwnerRole() == ROLE_AutonomousProxy)
 	{
 		SavedMoves.Reserve(MaxSavedMoveCount);
@@ -11152,7 +11119,6 @@ FNetworkPredictionData_Server_Character::FNetworkPredictionData_Server_Character
 	, ServerAccumulatedClientTimeStamp(0.0)
 	, LastUpdateTime(0.f)
 	, ServerTimeStampLastServerMove(0.f)
-	, MaxResponseTime(0.125f) // Deprecated, use MaxMoveDeltaTime instead
 	, MaxMoveDeltaTime(0.125f)
 	, bForceClientUpdate(false)
 	, LifetimeRawTimeDiscrepancy(0.f)
@@ -11178,8 +11144,6 @@ FNetworkPredictionData_Server_Character::FNetworkPredictionData_Server_Character
 		WorldCreationTime = World->GetTimeSeconds();
 		ServerTimeStamp = World->GetTimeSeconds();
 	}
-
-	MaxResponseTime = MaxMoveDeltaTime; // Deprecated, use MaxMoveDeltaTime instead
 }
 
 PRAGMA_ENABLE_DEPRECATION_WARNINGS // For deprecated members of FNetworkPredictionData_Server_Character
diff --git a/Engine/Source/Runtime/Engine/Private/Components/MovementComponent.cpp b/Engine/Source/Runtime/Engine/Private/Components/MovementComponent.cpp
--- a/Engine/Source/Runtime/Engine/Private/Components/MovementComponent.cpp
+++ b/Engine/Source/Runtime/Engine/Private/Components/MovementComponent.cpp
@@ -816,34 +816,4 @@ void UMovementComponent::AddRadialImpulse(const FVector& Origin, float Radius, f
 	// Default implementation does nothing
 }
 
-// TODO: Deprecated, remove.
-float UMovementComponent::GetMaxSpeedModifier() const
-{
-	return 1.0f;
-}
-
-// TODO: Deprecated, remove.
-float UMovementComponent::K2_GetMaxSpeedModifier() const
-{
-	// Allow calling old deprecated function to maintain old behavior until it is removed.
-	PRAGMA_DISABLE_DEPRECATION_WARNINGS
-	return GetMaxSpeedModifier();
-	PRAGMA_ENABLE_DEPRECATION_WARNINGS
-}
-
-// TODO: Deprecated, remove.
-float UMovementComponent::GetModifiedMaxSpeed() const
-{
-	return GetMaxSpeed();
-}
-
-// TODO: Deprecated, remove.
-float UMovementComponent::K2_GetModifiedMaxSpeed() const
-{
-	// Allow calling old deprecated function to maintain old behavior until it is removed.
-	PRAGMA_DISABLE_DEPRECATION_WARNINGS
-	return GetModifiedMaxSpeed();
-	PRAGMA_ENABLE_DEPRECATION_WARNINGS
-}
-
 #undef LOCTEXT_NAMESPACE
diff --git a/Engine/Source/Runtime/Engine/Private/Components/PawnMovementComponent.cpp b/Engine/Source/Runtime/Engine/Private/Components/PawnMovementComponent.cpp
--- a/Engine/Source/Runtime/Engine/Private/Components/PawnMovementComponent.cpp
+++ b/Engine/Source/Runtime/Engine/Private/Components/PawnMovementComponent.cpp
@@ -88,12 +88,6 @@ void UPawnMovementComponent::RequestPathMove(const FVector& MoveInput)
 	}
 }
 
-// TODO: deprecated, remove
-FVector UPawnMovementComponent::K2_GetInputVector() const
-{
-	return GetPendingInputVector();
-}
-
 void UPawnMovementComponent::OnTeleported()
 {
 	if (PawnOwner && PawnOwner->IsNetMode(NM_Client) && PawnOwner->IsLocallyControlled())
diff --git a/Engine/Source/Runtime/Engine/Private/Components/SceneComponent.cpp b/Engine/Source/Runtime/Engine/Private/Components/SceneComponent.cpp
--- a/Engine/Source/Runtime/Engine/Private/Components/SceneComponent.cpp
+++ b/Engine/Source/Runtime/Engine/Private/Components/SceneComponent.cpp
@@ -1768,14 +1768,6 @@ void USceneComponent::SetupAttachment(class USceneComponent* InParent, FName InS
 	}
 }
 
-//This function is used for giving AttachTo different bWeldSimulatedBodies default, but only when called from BP
-bool USceneComponent::K2_AttachTo(class USceneComponent* InParent, FName InSocketName, EAttachLocation::Type AttachLocationType, bool bWeldSimulatedBodies /*= true*/)
-{
-PRAGMA_DISABLE_DEPRECATION_WARNINGS
-	return AttachTo(InParent, InSocketName, AttachLocationType, bWeldSimulatedBodies);
-PRAGMA_ENABLE_DEPRECATION_WARNINGS
-}
-
 //This function is used for giving AttachToComponent different bWeldSimulatedBodies default, but only when called from BP
 bool USceneComponent::K2_AttachToComponent(USceneComponent* Parent, FName SocketName, EAttachmentRule LocationRule, EAttachmentRule RotationRule, EAttachmentRule ScaleRule, bool bWeldSimulatedBodies)
 {
@@ -1812,14 +1804,6 @@ void USceneComponent::ConvertAttachLocation(EAttachLocation::Type InAttachLocati
 	}
 }
 
-bool USceneComponent::AttachTo(class USceneComponent* Parent, FName InSocketName, EAttachLocation::Type AttachType /*= EAttachLocation::KeepRelativeOffset */, bool bWeldSimulatedBodies /*= false*/)
-{
-	FAttachmentTransformRules AttachmentRules(EAttachmentRule::KeepRelative, bWeldSimulatedBodies);
-	ConvertAttachLocation(AttachType, AttachmentRules.LocationRule, AttachmentRules.RotationRule, AttachmentRules.ScaleRule);
-
-	return AttachToComponent(Parent, AttachmentRules, InSocketName);
-}
-
 bool USceneComponent::AttachToComponent(USceneComponent* Parent, const FAttachmentTransformRules& AttachmentRules, FName SocketName)
 {
 	FUObjectThreadContext& ThreadContext = FUObjectThreadContext::Get();
@@ -2135,26 +2119,6 @@ bool USceneComponent::AttachToComponent(USceneComponent* Parent, const FAttachme
 	return false;
 }
 
-bool USceneComponent::SnapTo(class USceneComponent* Parent, FName InSocketName)
-{
-	return AttachToComponent(Parent, FAttachmentTransformRules::SnapToTargetNotIncludingScale, InSocketName);
-}
-
-void USceneComponent::DetachFromParent(bool bMaintainWorldPosition, bool bCallModify)
-{
-	FDetachmentTransformRules DetachmentRules(EDetachmentRule::KeepRelative, bCallModify);
-	if (bMaintainWorldPosition)
-	{
-		DetachmentRules.LocationRule = EDetachmentRule::KeepWorld;
-
-		// force maintain world rotation and scale for backwards compatibility
-		DetachmentRules.RotationRule = EDetachmentRule::KeepWorld;
-		DetachmentRules.ScaleRule = EDetachmentRule::KeepWorld;
-	}
-
-	DetachFromComponent(DetachmentRules);
-}
-
 void USceneComponent::K2_DetachFromComponent(EDetachmentRule LocationRule /*= EDetachmentRule::KeepRelative*/, EDetachmentRule RotationRule /*= EDetachmentRule::KeepRelative*/, EDetachmentRule ScaleRule /*= EDetachmentRule::KeepRelative*/, bool bCallModify /*= true*/)
 {
 	DetachFromComponent(FDetachmentTransformRules(LocationRule, RotationRule, ScaleRule, bCallModify));
diff --git a/Engine/Source/Runtime/Engine/Private/Components/TextRenderComponent.cpp b/Engine/Source/Runtime/Engine/Private/Components/TextRenderComponent.cpp
--- a/Engine/Source/Runtime/Engine/Private/Components/TextRenderComponent.cpp
+++ b/Engine/Source/Runtime/Engine/Private/Components/TextRenderComponent.cpp
@@ -1214,11 +1214,6 @@ FMatrix UTextRenderComponent::GetRenderMatrix() const
 	return UPrimitiveComponent::GetRenderMatrix();
 }
 
-void UTextRenderComponent::SetText(const FString& Value)
-{
-	K2_SetText(FText::FromString(Value));
-}
-
 void UTextRenderComponent::SetText(const FText& Value)
 {
 	K2_SetText(Value);
diff --git a/Engine/Source/Runtime/Engine/Private/Controller.cpp b/Engine/Source/Runtime/Engine/Private/Controller.cpp
--- a/Engine/Source/Runtime/Engine/Private/Controller.cpp
+++ b/Engine/Source/Runtime/Engine/Private/Controller.cpp
@@ -689,12 +689,6 @@ void AController::BeginInactiveState() {}
 
 void AController::EndInactiveState() {}
 
-
-APlayerController* AController::CastToPlayerController()
-{
-	return Cast<APlayerController>(this);
-}
-
 APawn* AController::K2_GetPawn() const
 {
 	return GetPawn();
diff --git a/Engine/Source/Runtime/Engine/Private/GameInstance.cpp b/Engine/Source/Runtime/Engine/Private/GameInstance.cpp
--- a/Engine/Source/Runtime/Engine/Private/GameInstance.cpp
+++ b/Engine/Source/Runtime/Engine/Private/GameInstance.cpp
@@ -316,26 +316,8 @@ FGameInstancePIEResult UGameInstance::InitializeForPlayInEditor(int32 PIEInstanc
 
 	Init();
 
-	// Give the deprecated method a chance to fail as well
-	FGameInstancePIEResult InitResult = FGameInstancePIEResult::Success();
-
-	if (InitResult.IsSuccess())
-	{
-		PRAGMA_DISABLE_DEPRECATION_WARNINGS
-		InitResult = InitializePIE(Params.bAnyBlueprintErrors, PIEInstanceIndex, Params.bRunAsDedicated) ?
-			FGameInstancePIEResult::Success() :
-			FGameInstancePIEResult::Failure(NSLOCTEXT("UnrealEd", "Error_CouldntInitInstance", "The game instance failed to Play/Simulate In Editor"));
-		PRAGMA_ENABLE_DEPRECATION_WARNINGS
-	}
-
-	return InitResult;
-}
-
-
-bool UGameInstance::InitializePIE(bool bAnyBlueprintErrors, int32 PIEInstance, bool bRunAsDedicated)
-{
-	// DEPRECATED VERSION
-	return true;
+	// Games can override this to return failure if PIE is not allowed for some reason
+	return FGameInstancePIEResult::Success();
 }
 
 FGameInstancePIEResult UGameInstance::StartPlayInEditorGameInstance(ULocalPlayer* LocalPlayer, const FGameInstancePIEParameters& Params)
@@ -483,31 +465,14 @@ FGameInstancePIEResult UGameInstance::StartPlayInEditorGameInstance(ULocalPlayer
 		PlayWorld->BeginPlay();
 	}
 
-	// Give the deprecated method a chance to fail as well
-	FGameInstancePIEResult StartResult = FGameInstancePIEResult::Success();
-
-	if (StartResult.IsSuccess())
-	{
-		PRAGMA_DISABLE_DEPRECATION_WARNINGS
-		StartResult = StartPIEGameInstance(LocalPlayer, Params.bSimulateInEditor, Params.bAnyBlueprintErrors, Params.bStartInSpectatorMode) ?
-			FGameInstancePIEResult::Success() :
-			FGameInstancePIEResult::Failure(NSLOCTEXT("UnrealEd", "Error_CouldntInitInstance", "The game instance failed to Play/Simulate In Editor"));
-		PRAGMA_ENABLE_DEPRECATION_WARNINGS
-	}
-
-	return StartResult;
+	// Games can override this to return failure if PIE is not allowed for some reason
+	return FGameInstancePIEResult::Success();
 }
 
 FGameInstancePIEResult UGameInstance::PostCreateGameModeForPIE(const FGameInstancePIEParameters& Params, AGameModeBase* GameMode)
 {
 	return FGameInstancePIEResult::Success();
 }
-
-bool UGameInstance::StartPIEGameInstance(ULocalPlayer* LocalPlayer, bool bInSimulateInEditor, bool bAnyBlueprintErrors, bool bStartInSpectatorMode)
-{
-	// DEPRECATED VERSION
-	return true;
-}
 #endif
 
 
diff --git a/Engine/Source/Runtime/Engine/Private/GameMode.cpp b/Engine/Source/Runtime/Engine/Private/GameMode.cpp
--- a/Engine/Source/Runtime/Engine/Private/GameMode.cpp
+++ b/Engine/Source/Runtime/Engine/Private/GameMode.cpp
@@ -44,23 +44,6 @@ AGameMode::AGameMode(const FObjectInitializer& ObjectInitializer)
 	MaxInactivePlayers = 16;
 }
 
-FString AGameMode::GetDefaultGameClassPath(const FString& MapName, const FString& Options, const FString& Portal) const
-{
-	// This is called on the CDO
-	return GetClass()->GetPathName();
-}
-
-TSubclassOf<AGameMode> AGameMode::GetGameModeClass(const FString& MapName, const FString& Options, const FString& Portal) const
-{
-	// This is called on the CDO
-	return GetClass();
-}
-
-FString AGameMode::StaticGetFullGameClassName(FString const& Str)
-{
-	return UGameMapsSettings::GetGameModeForName(Str);
-}
-
 FString AGameMode::GetNetworkNumber()
 {
 	UNetDriver* NetDriver = GetWorld()->GetNetDriver();
@@ -537,11 +520,6 @@ int32 AGameMode::GetNumSpectators()
 	return NumSpectators;
 }
 
-void AGameMode::StartNewPlayer(APlayerController* NewPlayer)
-{
-
-}
-
 void AGameMode::HandleStartingNewPlayer_Implementation(APlayerController* NewPlayer)
 {
 	// If players should start as spectators, leave them in the spectator state
@@ -574,11 +552,6 @@ bool AGameMode::PlayerCanRestart_Implementation(APlayerController* Player)
 	return Super::PlayerCanRestart_Implementation(Player);
 }
 
-void AGameMode::SendPlayer( APlayerController* aPlayer, const FString& FURL )
-{
-	aPlayer->ClientTravel( FURL, TRAVEL_Relative );
-}
-
 bool AGameMode::GetTravelType()
 {
 	return false;
diff --git a/Engine/Source/Runtime/Engine/Private/GameplayStatics.cpp b/Engine/Source/Runtime/Engine/Private/GameplayStatics.cpp
--- a/Engine/Source/Runtime/Engine/Private/GameplayStatics.cpp
+++ b/Engine/Source/Runtime/Engine/Private/GameplayStatics.cpp
@@ -643,13 +643,6 @@ class AActor* UGameplayStatics::BeginSpawningActorFromBlueprint(const UObject* W
 	return nullptr;
 }
 
-// deprecated
-class AActor* UGameplayStatics::BeginSpawningActorFromClass(const UObject* WorldContextObject, TSubclassOf<AActor> ActorClass, const FTransform& SpawnTransform, bool bNoCollisionFail /*= false*/, AActor* Owner /*= nullptr*/)
-{
-	ESpawnActorCollisionHandlingMethod const CollisionHandlingOverride = bNoCollisionFail ? ESpawnActorCollisionHandlingMethod::AdjustIfPossibleButDontSpawnIfColliding : ESpawnActorCollisionHandlingMethod::AlwaysSpawn;
-	return BeginDeferredActorSpawnFromClass(WorldContextObject, ActorClass, SpawnTransform, CollisionHandlingOverride, Owner);
-}
-
 class AActor* UGameplayStatics::BeginDeferredActorSpawnFromClass(const UObject* WorldContextObject, TSubclassOf<AActor> ActorClass, const FTransform& SpawnTransform, ESpawnActorCollisionHandlingMethod CollisionHandlingOverride /*= ESpawnActorCollisionHandlingMethod::Undefined*/, AActor* Owner /*= nullptr*/)
 {
 	SCOPE_CYCLE_COUNTER(STAT_SpawnTime);
@@ -675,12 +668,12 @@ class AActor* UGameplayStatics::BeginDeferredActorSpawnFromClass(const UObject*
 		else
 		{
 			//@TODO: RuntimeErrors: Overlogging
-			UE_LOG(LogScript, Warning, TEXT("UGameplayStatics::BeginSpawningActorFromClass: %s can not be spawned in NULL world"), *Class->GetName());		
+			UE_LOG(LogScript, Warning, TEXT("UGameplayStatics::BeginDeferredActorSpawnFromClass: %s can not be spawned in NULL world"), *Class->GetName());		
 		}
 	}
 	else
 	{
-		UE_LOG(LogScript, Warning, TEXT("UGameplayStatics::BeginSpawningActorFromClass: can not spawn an actor from a NULL class"));
+		UE_LOG(LogScript, Warning, TEXT("UGameplayStatics::BeginDeferredActorSpawnFromClass: can not spawn an actor from a NULL class"));
 	}
 	return nullptr;
 }
@@ -2520,45 +2513,6 @@ bool UGameplayStatics::PredictProjectilePath(const UObject* WorldContextObject,
 	return bBlockingHit;
 }
 
-
-// TODO: Deprecated, remove
-bool UGameplayStatics::PredictProjectilePath(
-	const UObject* WorldContextObject,
-	FHitResult& OutHit,
-	TArray<FVector>& OutPathPositions,
-	FVector& OutLastTraceDestination,
-	FVector StartPos,
-	FVector LaunchVelocity,
-	bool bTracePath,
-	float ProjectileRadius,
-	const TArray<TEnumAsByte<EObjectTypeQuery> >& ObjectTypes,
-	bool bTraceComplex,
-	const TArray<AActor*>& ActorsToIgnore,
-	EDrawDebugTrace::Type DrawDebugType,
-	float DrawDebugTime,
-	float SimFrequency,
-	float MaxSimTime,
-	float OverrideGravityZ)
-{
-	return Blueprint_PredictProjectilePath_ByObjectType(
-		WorldContextObject,
-		OutHit,
-		OutPathPositions,
-		OutLastTraceDestination,
-		StartPos,
-		LaunchVelocity,
-		bTracePath,
-		ProjectileRadius,
-		ObjectTypes,
-		bTraceComplex,
-		ActorsToIgnore,
-		DrawDebugType,
-		DrawDebugTime,
-		SimFrequency,
-		MaxSimTime,
-		OverrideGravityZ);
-}
-
 bool UGameplayStatics::Blueprint_PredictProjectilePath_Advanced(const UObject* WorldContextObject, const FPredictProjectilePathParams& PredictParams, FPredictProjectilePathResult& PredictResult)
 {
 	return PredictProjectilePath(WorldContextObject, PredictParams, PredictResult);
diff --git a/Engine/Source/Runtime/Engine/Private/HUD.cpp b/Engine/Source/Runtime/Engine/Private/HUD.cpp
--- a/Engine/Source/Runtime/Engine/Private/HUD.cpp
+++ b/Engine/Source/Runtime/Engine/Private/HUD.cpp
@@ -61,10 +61,6 @@ TAutoConsoleVariable<int32> GMaxDebugTextStringsPerActorCVar(
 	128,
 	TEXT("The maximum number of debug strings that can be attached to a given actor (<=0 : no limit)"));
 
-const FColor AHUD::WhiteColor(255, 255, 255, 255);
-const FColor AHUD::GreenColor(0, 255, 0, 255);
-const FColor AHUD::RedColor(255, 0, 0, 255);
-
 AHUD::AHUD(const FObjectInitializer& ObjectInitializer)
 	: Super(ObjectInitializer)
 {
diff --git a/Engine/Source/Runtime/Engine/Private/Level.cpp b/Engine/Source/Runtime/Engine/Private/Level.cpp
--- a/Engine/Source/Runtime/Engine/Private/Level.cpp
+++ b/Engine/Source/Runtime/Engine/Private/Level.cpp
@@ -2814,11 +2814,6 @@ void ULevel::RemoveUserDataOfClass(TSubclassOf<UAssetUserData> InUserDataClass)
 	}
 }
 
-bool ULevel::HasVisibilityRequestPending() const
-{
-	return (OwningWorld && this == OwningWorld->GetCurrentLevelPendingVisibility());
-}
-
 bool ULevel::HasVisibilityChangeRequestPending() const
 {
 	return (OwningWorld && ( this == OwningWorld->GetCurrentLevelPendingVisibility() || this == OwningWorld->GetCurrentLevelPendingInvisibility() ) );
diff --git a/Engine/Source/Runtime/Engine/Private/LevelTick.cpp b/Engine/Source/Runtime/Engine/Private/LevelTick.cpp
--- a/Engine/Source/Runtime/Engine/Private/LevelTick.cpp
+++ b/Engine/Source/Runtime/Engine/Private/LevelTick.cpp
@@ -1697,37 +1697,6 @@ void UWorld::Tick( ELevelTick TickType, float DeltaSeconds )
 	EndTickDrawEvent(TickDrawEvent);
 }
 
-/**
- *  Requests a one frame delay of Garbage Collection
- */
-void UWorld::DelayGarbageCollection()
-{
-	GEngine->DelayGarbageCollection();
-}
-
-void UWorld::ForceGarbageCollection( bool bFullPurge)
-{
-	GEngine->ForceGarbageCollection(bFullPurge);
-}
-
-void UWorld::SetTimeUntilNextGarbageCollection(const float MinTimeUntilNextPass)
-{
-	GEngine->SetTimeUntilNextGarbageCollection(MinTimeUntilNextPass);
-}
-
-float UWorld::GetTimeBetweenGarbageCollectionPasses() const
-{
-	return GEngine->GetTimeBetweenGarbageCollectionPasses();
-}
-
-/**
- *  Interface to allow WorldSettings to request immediate garbage collection
- */
-void UWorld::PerformGarbageCollectionAndCleanupActors()
-{
-	GEngine->PerformGarbageCollectionAndCleanupActors();
-}
-
 void UWorld::CleanupActors()
 {
 	// Remove NULL entries from actor list. Only does so for dynamic actors to avoid resorting; in theory static 
diff --git a/Engine/Source/Runtime/Engine/Private/Pawn.cpp b/Engine/Source/Runtime/Engine/Private/Pawn.cpp
--- a/Engine/Source/Runtime/Engine/Private/Pawn.cpp
+++ b/Engine/Source/Runtime/Engine/Private/Pawn.cpp
@@ -651,12 +651,6 @@ FVector APawn::GetLastMovementInputVector() const
 	return LastControlInputVector;
 }
 
-// TODO: deprecated, remove
-FVector APawn::K2_GetMovementInputVector() const
-{
-	return GetPendingMovementInputVector();
-}
-
 FVector APawn::ConsumeMovementInputVector()
 {
 	UPawnMovementComponent* MovementComponent = GetMovementComponent();
@@ -1013,16 +1007,6 @@ void APawn::GetMoveGoalReachTest(const AActor* MovingActor, const FVector& MoveO
 	GetSimpleCollisionCylinder(GoalRadius, GoalHalfHeight);
 }
 
-// @TODO: Deprecated, remove me.
-void APawn::LaunchPawn(FVector LaunchVelocity, bool bXYOverride, bool bZOverride)
-{
-	ACharacter* Character = Cast<ACharacter>(this);
-	if (Character)
-	{
-		Character->LaunchCharacter(LaunchVelocity, bXYOverride, bZOverride);
-	}
-}
-
 // REPLICATION
 
 void APawn::PostNetReceiveVelocity(const FVector& NewVelocity)
diff --git a/Engine/Source/Runtime/Engine/Private/UnrealEngine.cpp b/Engine/Source/Runtime/Engine/Private/UnrealEngine.cpp
--- a/Engine/Source/Runtime/Engine/Private/UnrealEngine.cpp
+++ b/Engine/Source/Runtime/Engine/Private/UnrealEngine.cpp
@@ -12046,11 +12046,6 @@ bool UEngine::HandleServerTravelCommand( const TCHAR* Cmd, FOutputDevice& Ar, UW
 	}
 	return false;
 }
-
-bool UEngine::HandleSayCommand( const TCHAR* Cmd, FOutputDevice& Ar, UWorld* InWorld )
-{
-	return false;
-}
 #endif
 
 bool UEngine::HandleDisconnectCommand( const TCHAR* Cmd, FOutputDevice& Ar, UWorld *InWorld )
diff --git a/Engine/Source/Runtime/Engine/Public/CollisionQueryParams.h b/Engine/Source/Runtime/Engine/Public/CollisionQueryParams.h
--- a/Engine/Source/Runtime/Engine/Public/CollisionQueryParams.h
+++ b/Engine/Source/Runtime/Engine/Public/CollisionQueryParams.h
@@ -136,32 +136,6 @@ public:
 
 	// Constructors
 #if !FIND_UNKNOWN_SCENE_QUERIES
-	/** 
-	 *  DEPRECATED!  
-	 *  Please instead provide a FName parameter when constructing a FCollisionQueryParams object which will use the other constructor.
-	 *  Providing a single string literal argument, such as TEXT("foo"), instead of an explicit FNAME
-	 *  can cause this constructor to be invoked instead the the other which was likely the programmers intention. 
-	 *  This constructor will eventually be deprecated to avoid this potentially ambiguous case.
-	 */ 
-	UE_DEPRECATED(4.11, "FCollisionQueryParams, to avoid ambiguity, please use other constructor and explicitly provide an FName parameter (not just a string literal) as the first parameter")
-	FCollisionQueryParams(bool bInTraceComplex)
-	{
-		bTraceComplex = bInTraceComplex;
-		MobilityType = EQueryMobilityType::Any;
-		TraceTag = NAME_None;
-		bFindInitialOverlaps = true;
-		bReturnFaceIndex = false;
-		bReturnPhysicalMaterial = false;
-		bComponentListUnique = true;
-		IgnoreMask = 0;
-		bIgnoreBlocks = false;
-		bIgnoreTouches = false;
-		bSkipNarrowPhase = false;
-		StatId = GetUnknownStatId();
-#if !(UE_BUILD_TEST || UE_BUILD_SHIPPING)
-		bDebugQuery = false;
-#endif
-	}
 
 	FCollisionQueryParams()
 	{
diff --git a/Engine/Source/Runtime/GameplayTags/Classes/GameplayTagContainer.h b/Engine/Source/Runtime/GameplayTags/Classes/GameplayTagContainer.h
--- a/Engine/Source/Runtime/GameplayTags/Classes/GameplayTagContainer.h
+++ b/Engine/Source/Runtime/GameplayTags/Classes/GameplayTagContainer.h
@@ -17,15 +17,9 @@ GAMEPLAYTAGS_API DECLARE_LOG_CATEGORY_EXTERN(LogGameplayTags, Log, All);
 
 DECLARE_STATS_GROUP(TEXT("Gameplay Tags"), STATGROUP_GameplayTags, STATCAT_Advanced);
 
-DECLARE_CYCLE_STAT_EXTERN(TEXT("FGameplayTagContainer::HasTag"), STAT_FGameplayTagContainer_HasTag, STATGROUP_GameplayTags, GAMEPLAYTAGS_API);
 DECLARE_CYCLE_STAT_EXTERN(TEXT("FGameplayTagContainer::DoesTagContainerMatch"), STAT_FGameplayTagContainer_DoesTagContainerMatch, STATGROUP_GameplayTags, GAMEPLAYTAGS_API);
-DECLARE_CYCLE_STAT_EXTERN(TEXT("UGameplayTagsManager::GameplayTagsMatch"), STAT_UGameplayTagsManager_GameplayTagsMatch, STATGROUP_GameplayTags, GAMEPLAYTAGS_API);
 
-struct FGameplayTagContainer;
-
-// DEPRECATED ENUMS
-UENUM(BlueprintType)
-namespace EGameplayTagMatchType
+namespace UE_DEPRECATED(5.0, "Deprecated in favor of HasExact and related functions") EGameplayTagMatchType
 {
 	enum Type
 	{
@@ -37,8 +31,11 @@ namespace EGameplayTagMatchType
 UENUM(BlueprintType)
 enum class EGameplayContainerMatchType : uint8
 {
-	Any,	//	Means the filter is populated by any tag matches in this container.
-	All		//	Means the filter is only populated if all of the tags in this container match.
+	//	Means the filter is populated by any tag matches in this container.
+	Any,	
+
+	//	Means the filter is only populated if all of the tags in this container match.
+	All		
 };
 
 typedef uint16 FGameplayTagNetIndex;
@@ -202,47 +199,6 @@ struct GAMEPLAYTAGS_API FGameplayTag
 	/** An empty Gameplay Tag */
 	static const FGameplayTag EmptyTag;
 
-	// DEPRECATED
-
-PRAGMA_DISABLE_DEPRECATION_WARNINGS
-
-	/**
-	 * Check to see if two FGameplayTags match with explicit match types
-	 *
-	 * @param MatchTypeOne	How we compare this tag, Explicitly or a match with any parents as well
-	 * @param Other			The second tag to compare against
-	 * @param MatchTypeTwo	How we compare Other tag, Explicitly or a match with any parents as well
-	 * 
-	 * @return True if there is a match according to the specified match types; false if not
-	 */
-	UE_DEPRECATED(4.15, "Deprecated in favor of MatchesTag")
-	FORCEINLINE_DEBUGGABLE bool Matches(TEnumAsByte<EGameplayTagMatchType::Type> MatchTypeOne, const FGameplayTag& Other, TEnumAsByte<EGameplayTagMatchType::Type> MatchTypeTwo) const
-	{
-		bool bResult = false;
-		if (MatchTypeOne == EGameplayTagMatchType::Explicit && MatchTypeTwo == EGameplayTagMatchType::Explicit)
-		{
-			bResult = TagName == Other.TagName;
-		}
-		else
-		{
-			bResult = ComplexMatches(MatchTypeOne, Other, MatchTypeTwo);
-		}
-		return bResult;
-	}
-	/**
-	 * Check to see if two FGameplayTags match
-	 *
-	 * @param MatchTypeOne	How we compare this tag, Explicitly or a match with any parents as well
-	 * @param Other			The second tag to compare against
-	 * @param MatchTypeTwo	How we compare Other tag, Explicitly or a match with any parents as well
-	 * 
-	 * @return True if there is a match according to the specified match types; false if not
-	 */
-	UE_DEPRECATED(4.15, "Deprecated in favor of MatchesTag")
-	bool ComplexMatches(TEnumAsByte<EGameplayTagMatchType::Type> MatchTypeOne, const FGameplayTag& Other, TEnumAsByte<EGameplayTagMatchType::Type> MatchTypeTwo) const;
-
-PRAGMA_ENABLE_DEPRECATION_WARNINGS
-
 protected:
 
 	/** Intentionally private so only the tag manager can use */
@@ -628,58 +584,9 @@ struct GAMEPLAYTAGS_API FGameplayTagContainer
 	/** An empty Gameplay Tag Container */
 	static const FGameplayTagContainer EmptyContainer;
 
-	// DEPRECATED FUNCTIONALITY
-
 PRAGMA_DISABLE_DEPRECATION_WARNINGS
 
-	UE_DEPRECATED(4.15, "Deprecated in favor of Reset")
-		void RemoveAllTags(int32 Slack = 0)
-	{
-		Reset(Slack);
-	}
-
-	UE_DEPRECATED(4.15, "Deprecated in favor of Reset")
-		void RemoveAllTagsKeepSlack()
-	{
-		Reset();
-	}
-
-	/**
-	 * Determine if the container has the specified tag. This forces an explicit match. 
-	 * This function exists for convenience and brevity. We do not wish to use default values for ::HasTag match type parameters, to avoid confusion on what the default behavior is. (E.g., we want people to think and use the right match type).
-	 * 
-	 * @param TagToCheck			Tag to check if it is present in the container
-	 * 
-	 * @return True if the tag is in the container, false if it is not
-	 */
-	UE_DEPRECATED(4.15, "Deprecated in favor of HasTagExact")
-	FORCEINLINE_DEBUGGABLE bool HasTagExplicit(FGameplayTag const& TagToCheck) const
-	{
-		return HasTag(TagToCheck, EGameplayTagMatchType::Explicit, EGameplayTagMatchType::Explicit);
-	}
-
-	/**
-	 * Determine if the container has the specified tag
-	 * 
-	 * @param TagToCheck			Tag to check if it is present in the container
-	 * @param TagMatchType			Type of match to use for the tags in this container
-	 * @param TagToCheckMatchType	Type of match to use for the TagToCheck Param
-	 * 
-	 * @return True if the tag is in the container, false if it is not
-	 */
-	UE_DEPRECATED(4.15, "Deprecated in favor of HasTag with no parameters")
-	FORCEINLINE_DEBUGGABLE bool HasTag(FGameplayTag const& TagToCheck, TEnumAsByte<EGameplayTagMatchType::Type> TagMatchType, TEnumAsByte<EGameplayTagMatchType::Type> TagToCheckMatchType) const
-	{
-		SCOPE_CYCLE_COUNTER(STAT_FGameplayTagContainer_HasTag);
-		if (!TagToCheck.IsValid())
-		{
-			return false;
-		}
-
-		return HasTagFast(TagToCheck, TagMatchType, TagToCheckMatchType);
-	}
-
-	/** Version of above that is called from conditions where you know tag is valid */
+	UE_DEPRECATED(5.0, "Deprecated in favor of HasTag or HasTagExact")
 	FORCEINLINE_DEBUGGABLE bool HasTagFast(FGameplayTag const& TagToCheck, TEnumAsByte<EGameplayTagMatchType::Type> TagMatchType, TEnumAsByte<EGameplayTagMatchType::Type> TagToCheckMatchType) const
 	{
 		bool bResult;
@@ -701,65 +608,10 @@ PRAGMA_DISABLE_DEPRECATION_WARNINGS
 		return bResult;
 	}
 
-	/**
-	 * Determine if the container has the specified tag
-	 * 
-	 * @param TagToCheck			Tag to check if it is present in the container
-	 * @param TagMatchType			Type of match to use for the tags in this container
-	 * @param TagToCheckMatchType	Type of match to use for the TagToCheck Param
-	 * 
-	 * @return True if the tag is in the container, false if it is not
-	 */
+	UE_DEPRECATED(5.0, "Deprecated in favor of HasTag or HasTagExact")
 	bool ComplexHasTag(FGameplayTag const& TagToCheck, TEnumAsByte<EGameplayTagMatchType::Type> TagMatchType, TEnumAsByte<EGameplayTagMatchType::Type> TagToCheckMatchType) const;
 
-	/**
-	 * Checks if this container matches ANY of the tags in the specified container. Performs matching by expanding this container out
-	 * to include its parent tags.
-	 *
-	 * @param Other					Container we are checking against
-	 * @param bCountEmptyAsMatch	If true, the parameter tag container will count as matching even if it's empty
-	 *
-	 * @return True if this container has ANY the tags of the passed in container
-	 */
-	UE_DEPRECATED(4.15, "Deprecated in favor of HasAny")
-	FORCEINLINE_DEBUGGABLE bool MatchesAny(const FGameplayTagContainer& Other, bool bCountEmptyAsMatch) const
-	{
-		if (Other.IsEmpty())
-		{
-			return bCountEmptyAsMatch;
-		}
-		return DoesTagContainerMatch(Other, EGameplayTagMatchType::IncludeParentTags, EGameplayTagMatchType::Explicit, EGameplayContainerMatchType::Any);
-	}
-
-	/**
-	 * Checks if this container matches ALL of the tags in the specified container. Performs matching by expanding this container out to
-	 * include its parent tags.
-	 *
-	 * @param Other				Container we are checking against
-	 * @param bCountEmptyAsMatch	If true, the parameter tag container will count as matching even if it's empty
-	 * 
-	 * @return True if this container has ALL the tags of the passed in container
-	 */
-	UE_DEPRECATED(4.15, "Deprecated in favor of HasAll")
-	FORCEINLINE_DEBUGGABLE bool MatchesAll(const FGameplayTagContainer& Other, bool bCountEmptyAsMatch) const
-	{
-		if (Other.IsEmpty())
-		{
-			return bCountEmptyAsMatch;
-		}
-		return DoesTagContainerMatch(Other, EGameplayTagMatchType::IncludeParentTags, EGameplayTagMatchType::Explicit, EGameplayContainerMatchType::All);
-	}
-
-	/**
-	 * Returns true if the tags in this container match the tags in OtherContainer for the specified matching types.
-	 *
-	 * @param OtherContainer		The Container to filter against
-	 * @param TagMatchType			Type of match to use for the tags in this container
-	 * @param OtherTagMatchType		Type of match to use for the tags in the OtherContainer param
-	 * @param ContainerMatchType	Type of match to use for filtering
-	 *
-	 * @return Returns true if ContainerMatchType is Any and any of the tags in OtherContainer match the tags in this or ContainerMatchType is All and all of the tags in OtherContainer match at least one tag in this. Returns false otherwise.
-	 */
+	UE_DEPRECATED(5.0, "Deprecated in favor of HasAll and related functions")
 	FORCEINLINE_DEBUGGABLE bool DoesTagContainerMatch(const FGameplayTagContainer& OtherContainer, TEnumAsByte<EGameplayTagMatchType::Type> TagMatchType, TEnumAsByte<EGameplayTagMatchType::Type> OtherTagMatchType, EGameplayContainerMatchType ContainerMatchType) const
 	{
 		SCOPE_CYCLE_COUNTER(STAT_FGameplayTagContainer_DoesTagContainerMatch);
@@ -793,34 +645,10 @@ PRAGMA_DISABLE_DEPRECATION_WARNINGS
 		return bResult;
 	}
 
-	/**
-	 * Returns a filtered version of this container, as if the container were filtered by matches from the parameter container
-	 *
-	 * @param OtherContainer		The Container to filter against
-	 * @param TagMatchType			Type of match to use for the tags in this container
-	 * @param OtherTagMatchType		Type of match to use for the tags in the OtherContainer param
-	 *
-	 * @return A FGameplayTagContainer containing the filtered tags
-	 */
-	UE_DEPRECATED(4.15, "Deprecated in favor of HasAll")
-	FGameplayTagContainer Filter(const FGameplayTagContainer& OtherContainer, TEnumAsByte<EGameplayTagMatchType::Type> TagMatchType, TEnumAsByte<EGameplayTagMatchType::Type> OtherTagMatchType) const;
-
 PRAGMA_ENABLE_DEPRECATION_WARNINGS
 
 protected:
 
-	/**
-	 * Returns true if the tags in this container match the tags in OtherContainer for the specified matching types.
-	 *
-	 * @param OtherContainer		The Container to filter against
-	 * @param TagMatchType			Type of match to use for the tags in this container
-	 * @param OtherTagMatchType		Type of match to use for the tags in the OtherContainer param
-	 * @param ContainerMatchType	Type of match to use for filtering
-	 *
-	 * @return Returns true if ContainerMatchType is Any and any of the tags in OtherContainer match the tags in this or ContainerMatchType is All and all of the tags in OtherContainer match at least one tag in this. Returns false otherwise.
-	 */
-	bool DoesTagContainerMatchComplex(const FGameplayTagContainer& OtherContainer, TEnumAsByte<EGameplayTagMatchType::Type> TagMatchType, TEnumAsByte<EGameplayTagMatchType::Type> OtherTagMatchType, EGameplayContainerMatchType ContainerMatchType) const;
-
 	/**
 	 * If a Tag with the specified tag name explicitly exists, it will remove that tag and return true.  Otherwise, it 
 	   returns false.  It does NOT check the TagName for validity (i.e. the tag could be obsolete and so not exist in
diff --git a/Engine/Source/Runtime/GameplayTags/Classes/GameplayTagsManager.h b/Engine/Source/Runtime/GameplayTags/Classes/GameplayTagsManager.h
--- a/Engine/Source/Runtime/GameplayTags/Classes/GameplayTagsManager.h
+++ b/Engine/Source/Runtime/GameplayTags/Classes/GameplayTagsManager.h
@@ -591,28 +591,6 @@ class GAMEPLAYTAGS_API UGameplayTagsManager : public UObject
 
 #endif //WITH_EDITOR
 
-	UE_DEPRECATED(4.15, "Call MatchesTag on FGameplayTag instead")
-	FORCEINLINE_DEBUGGABLE bool GameplayTagsMatch(const FGameplayTag& GameplayTagOne, TEnumAsByte<EGameplayTagMatchType::Type> MatchTypeOne, const FGameplayTag& GameplayTagTwo, TEnumAsByte<EGameplayTagMatchType::Type> MatchTypeTwo) const
-	{
-		SCOPE_CYCLE_COUNTER(STAT_UGameplayTagsManager_GameplayTagsMatch);
-		bool bResult = false;
-		if (MatchTypeOne == EGameplayTagMatchType::Explicit && MatchTypeTwo == EGameplayTagMatchType::Explicit)
-		{
-			bResult = GameplayTagOne == GameplayTagTwo;
-		}
-		else
-		{
-			// Convert both to their containers and do that match
-			const FGameplayTagContainer* ContainerOne = GetSingleTagContainer(GameplayTagOne);
-			const FGameplayTagContainer* ContainerTwo = GetSingleTagContainer(GameplayTagTwo);
-			if (ContainerOne && ContainerTwo)
-			{
-				bResult = ContainerOne->DoesTagContainerMatch(*ContainerTwo, MatchTypeOne, MatchTypeTwo, EGameplayContainerMatchType::Any);
-			}
-		}
-		return bResult;
-	}
-
 	void PrintReplicationIndices();
 
 #if !(UE_BUILD_SHIPPING || UE_BUILD_TEST)
diff --git a/Engine/Source/Runtime/GameplayTags/Private/GameplayTagContainer.cpp b/Engine/Source/Runtime/GameplayTags/Private/GameplayTagContainer.cpp
--- a/Engine/Source/Runtime/GameplayTags/Private/GameplayTagContainer.cpp
+++ b/Engine/Source/Runtime/GameplayTags/Private/GameplayTagContainer.cpp
@@ -15,9 +15,7 @@ const FGameplayTag FGameplayTag::EmptyTag;
 const FGameplayTagContainer FGameplayTagContainer::EmptyContainer;
 const FGameplayTagQuery FGameplayTagQuery::EmptyQuery;
 
-DEFINE_STAT(STAT_FGameplayTagContainer_HasTag);
 DEFINE_STAT(STAT_FGameplayTagContainer_DoesTagContainerMatch);
-DEFINE_STAT(STAT_UGameplayTagsManager_GameplayTagsMatch);
 
 static bool GEnableGameplayTagDetailedStats = false;
 static FAutoConsoleVariableRef CVarGameplayTagDetailedStats(TEXT("GameplayTags.EnableDetailedStats"), GEnableGameplayTagDetailedStats, TEXT("Runtime toggle for verbose CPU profiling stats"), ECVF_Default);
@@ -540,6 +538,8 @@ bool FGameplayTagContainer::operator!=(FGameplayTagContainer const& Other) const
 	return !operator==(Other);
 }
 
+PRAGMA_DISABLE_DEPRECATION_WARNINGS
+
 bool FGameplayTagContainer::ComplexHasTag(FGameplayTag const& TagToCheck, TEnumAsByte<EGameplayTagMatchType::Type> TagMatchType, TEnumAsByte<EGameplayTagMatchType::Type> TagToCheckMatchType) const
 {
 	check(TagMatchType != EGameplayTagMatchType::Explicit || TagToCheckMatchType != EGameplayTagMatchType::Explicit);
@@ -561,6 +561,8 @@ bool FGameplayTagContainer::ComplexHasTag(FGameplayTag const& TagToCheck, TEnumA
 	return false;
 }
 
+PRAGMA_ENABLE_DEPRECATION_WARNINGS
+
 DECLARE_CYCLE_STAT(TEXT("FGameplayTagContainer::RemoveTagByExplicitName"), STAT_FGameplayTagContainer_RemoveTagByExplicitName, STATGROUP_GameplayTags);
 
 bool FGameplayTagContainer::RemoveTagByExplicitName(const FName& TagName)
@@ -627,26 +629,6 @@ FGameplayTagContainer FGameplayTagContainer::GetGameplayTagParents() const
 
 DECLARE_CYCLE_STAT(TEXT("FGameplayTagContainer::Filter"), STAT_FGameplayTagContainer_Filter, STATGROUP_GameplayTags);
 
-FGameplayTagContainer FGameplayTagContainer::Filter(const FGameplayTagContainer& OtherContainer, TEnumAsByte<EGameplayTagMatchType::Type> TagMatchType, TEnumAsByte<EGameplayTagMatchType::Type> OtherTagMatchType) const
-{
-	SCOPE_CYCLE_COUNTER(STAT_FGameplayTagContainer_Filter);
-
-	FGameplayTagContainer ResultContainer;
-
-	for (const FGameplayTag& Tag : GameplayTags)
-	{
-		PRAGMA_DISABLE_DEPRECATION_WARNINGS
-		// Check to see if all of these tags match other container, with types swapped
-		if (OtherContainer.HasTag(Tag, OtherTagMatchType, TagMatchType))
-		{
-			ResultContainer.AddTagFast(Tag);
-		}
-		PRAGMA_ENABLE_DEPRECATION_WARNINGS
-	}
-
-	return ResultContainer;
-}
-
 FGameplayTagContainer FGameplayTagContainer::Filter(const FGameplayTagContainer& OtherContainer) const
 {
 	SCOPE_CYCLE_COUNTER(STAT_FGameplayTagContainer_Filter);
@@ -681,43 +663,6 @@ FGameplayTagContainer FGameplayTagContainer::FilterExact(const FGameplayTagConta
 	return ResultContainer;
 }
 
-bool FGameplayTagContainer::DoesTagContainerMatchComplex(const FGameplayTagContainer& OtherContainer, TEnumAsByte<EGameplayTagMatchType::Type> TagMatchType, TEnumAsByte<EGameplayTagMatchType::Type> OtherTagMatchType, EGameplayContainerMatchType ContainerMatchType) const
-{
-	UGameplayTagsManager& TagManager = UGameplayTagsManager::Get();
-
-	for (TArray<FGameplayTag>::TConstIterator OtherIt(OtherContainer.GameplayTags); OtherIt; ++OtherIt)
-	{
-		bool bTagFound = false;
-
-		for (TArray<FGameplayTag>::TConstIterator It(this->GameplayTags); It; ++It)
-		{
-			PRAGMA_DISABLE_DEPRECATION_WARNINGS
-			if (TagManager.GameplayTagsMatch(*It, TagMatchType, *OtherIt, OtherTagMatchType) == true)
-			{
-				if (ContainerMatchType == EGameplayContainerMatchType::Any)
-				{
-					return true;
-				}
-
-				bTagFound = true;
-
-				// we only need one match per tag in OtherContainer, so don't bother looking for more
-				break;
-			}
-			PRAGMA_ENABLE_DEPRECATION_WARNINGS
-		}
-
-		if (ContainerMatchType == EGameplayContainerMatchType::All && bTagFound == false)
-		{
-			return false;
-		}
-	}
-
-	// if we've reached this far then either we are looking for any match and didn't find one (return false) or we're looking for all matches and didn't miss one (return true).
-	check(ContainerMatchType == EGameplayContainerMatchType::All || ContainerMatchType == EGameplayContainerMatchType::Any);
-	return ContainerMatchType == EGameplayContainerMatchType::All;
-}
-
 bool FGameplayTagContainer::MatchesQuery(const FGameplayTagQuery& Query) const
 {
 	return Query.Matches(*this);
@@ -1100,13 +1045,6 @@ FText FGameplayTagContainer::ToMatchingText(EGameplayContainerMatchType MatchTyp
 	return FText::Format(MatchingDescription[DescriptionIndex], Arguments);
 }
 
-bool FGameplayTag::ComplexMatches(TEnumAsByte<EGameplayTagMatchType::Type> MatchTypeOne, const FGameplayTag& Other, TEnumAsByte<EGameplayTagMatchType::Type> MatchTypeTwo) const
-{
-	PRAGMA_DISABLE_DEPRECATION_WARNINGS
-	return UGameplayTagsManager::Get().GameplayTagsMatch(*this, MatchTypeOne, Other, MatchTypeTwo);
-	PRAGMA_ENABLE_DEPRECATION_WARNINGS
-}
-
 DECLARE_CYCLE_STAT(TEXT("FGameplayTag::GetSingleTagContainer"), STAT_FGameplayTag_GetSingleTagContainer, STATGROUP_GameplayTags);
 
 const FGameplayTagContainer& FGameplayTag::GetSingleTagContainer() const
diff --git a/Engine/Source/Runtime/GameplayTags/Public/GameplayTagsModule.h b/Engine/Source/Runtime/GameplayTags/Public/GameplayTagsModule.h
--- a/Engine/Source/Runtime/GameplayTags/Public/GameplayTagsModule.h
+++ b/Engine/Source/Runtime/GameplayTags/Public/GameplayTagsModule.h
@@ -46,17 +46,5 @@ public:
 	/** Delegate that gets called after the settings have changed in the editor */
 	static GAMEPLAYTAGS_API FSimpleMulticastDelegate OnTagSettingsChanged;
 
-	UE_DEPRECATED(4.15, "Call FGameplayTag::RequestGameplayTag or RequestGameplayTag on the manager instead")
-	FORCEINLINE_DEBUGGABLE static FGameplayTag RequestGameplayTag(FName InTagName, bool ErrorIfNotFound=true)
-	{
-		return UGameplayTagsManager::Get().RequestGameplayTag(InTagName, ErrorIfNotFound);
-	}
-
-	UE_DEPRECATED(4.15, "Call UGameplayTagsManager::Get instead")
-	FORCEINLINE_DEBUGGABLE static UGameplayTagsManager& GetGameplayTagsManager()
-	{
-		return UGameplayTagsManager::Get();
-	}
-
 };
 
