Cherry-picking native gameplay tag and tag filtering improvements from UE4 (CL# 15410246 + 15421371 + 15472831 and 15296748)

diff --git a/Engine/Source/Runtime/GameplayTags/Private/GameplayTagsManager.cpp b/Engine/Source/Runtime/GameplayTags/Private/GameplayTagsManager.cpp
--- a/Engine/Source/Runtime/GameplayTags/Private/GameplayTagsManager.cpp
+++ b/Engine/Source/Runtime/GameplayTags/Private/GameplayTagsManager.cpp
@@ -320,15 +320,6 @@ void UGameplayTagsManager::ConstructGameplayTagTree()
 			}
 		}
 
-		{
-			SCOPE_LOG_GAMEPLAYTAGS(TEXT("UGameplayTagsManager::ConstructGameplayTagTree: Add sourced native tags"));
-			// Add native tags from a named source.
-			for (const auto& KVP : NativeSourcesToAdd)
-			{
-				AddTagsFromNativeSource(KVP.Key, KVP.Value);
-			}
-		}
-
 		// If we didn't load any tables it might be async loading, so load again with a flush
 		if (GameplayTagTables.Num() == 0)
 		{
@@ -338,7 +329,7 @@ void UGameplayTagsManager::ConstructGameplayTagTree()
 		{
 			SCOPE_LOG_GAMEPLAYTAGS(TEXT("UGameplayTagsManager::ConstructGameplayTagTree: Construct from data asset"));
 			for (UDataTable* DataTable : GameplayTagTables)
-				{
+			{
 				if (DataTable)
 				{
 					PopulateTreeFromDataTable(DataTable);
@@ -1884,60 +1875,19 @@ void UGameplayTagsManager::AddNativeGameplayTag(FNativeGameplayTag* TagSource, F
 
 void UGameplayTagsManager::RemoveNativeGameplayTag(const FNativeGameplayTag* TagSource)
 {
-	//TODO - Removing tags isn't really a thing right now, but setting this up for the future.
-	//DestroyGameplayTagTree();
-	//ConstructGameplayTagTree();
-	//OnEditorRefreshGameplayTagTree.Broadcast();
-}
-
-void UGameplayTagsManager::AddNativeGameplayTagSource(const FString& InNativeSourceName, TSharedRef<const FNativeGameplayTagSource> NativeSource)
-{
-	const FName NativeSourceName = FName(*(FGameplayTagSource::GetNativeName().ToString() + TEXT("+") + InNativeSourceName));
-
-	// Create native source specific to a plugin+module
-	FindOrAddTagSource(NativeSourceName, EGameplayTagSourceType::Native);
-
-	// Replace the old one, in case we reload or something, allow new ones to come in.
-	NativeSourcesToAdd.Add(NativeSourceName, NativeSource);
-
-#if WITH_EDITOR
-	EditorRefreshGameplayTagTree();
-#else
-	AddTagsFromNativeSource(NativeSourceName, NativeSource);
-#endif
-}
-
-void UGameplayTagsManager::RemoveNativeGameplayTagSource(const FString& InNativeSourceName)
-{
-	const FName NativeSourceName = FName(*(FGameplayTagSource::GetNativeName().ToString() + TEXT("+") + InNativeSourceName));
-
-	if (NativeSourcesToAdd.Remove(NativeSourceName) > 0)
+	// TODO This is awful, need to invalidate the tag tree, not rebuild it.
 	{
 #if WITH_EDITOR
 		EditorRefreshGameplayTagTree();
 #else
-		//TODO - Removing tags isn't really a thing right now, but setting this up for the future.
+	//TODO - Removing tags isn't really a thing right now, but setting this up for the future.
+	//DestroyGameplayTagTree();
+	//ConstructGameplayTagTree();
+	//OnEditorRefreshGameplayTagTree.Broadcast();
 #endif
 	}
 }
 
-void UGameplayTagsManager::AddTagsFromNativeSource(FName NativeSourceName, TSharedPtr<const FNativeGameplayTagSource> NativeSource)
-{
-	if (!NativeSource.IsValid())
-	{
-		return;
-	}
-
-	for (const FGameplayTagTableRow& Row : NativeSource->NativeTags)
-	{
-		AddTagTableRow(Row, NativeSourceName);
-	}
-
-	InvalidateNetworkIndex();
-
-	IGameplayTagsModule::OnGameplayTagTreeChanged.Broadcast();
-}
-
 void UGameplayTagsManager::CallOrRegister_OnDoneAddingNativeTagsDelegate(FSimpleMulticastDelegate::FDelegate Delegate)
 {
 	if (bDoneAddingNativeTags)
diff --git a/Engine/Source/Runtime/GameplayTags/Private/NativeGameplayTags.cpp b/Engine/Source/Runtime/GameplayTags/Private/NativeGameplayTags.cpp
--- a/Engine/Source/Runtime/GameplayTags/Private/NativeGameplayTags.cpp
+++ b/Engine/Source/Runtime/GameplayTags/Private/NativeGameplayTags.cpp
@@ -7,39 +7,51 @@
 #include "PluginDescriptor.h"
 #include "Interfaces/IPluginManager.h"
 
-#define LOCTEXT_NAMESPACE "FNativeGameplayTagSource"
+#define LOCTEXT_NAMESPACE "FNativeGameplayTag"
 
-FGameplayTag FNativeGameplayTagSource::Add(FName TagName, const FString& TagDevComment)
-{
-	if (TagName.IsNone())
-	{
-		return FGameplayTag();
-	}
-
-	FGameplayTag NewTag = FGameplayTag(TagName);
-	NativeTags.Add(FGameplayTagTableRow(TagName, TagDevComment));
+TSet<const class FNativeGameplayTag*> FNativeGameplayTag::RegisteredNativeTags;
 
-	return NewTag;
-}
+#if !UE_BUILD_SHIPPING
 
-TSet<const class FNativeGameplayTag*> FNativeGameplayTag::RegisteredNativeTags;
+//FEditorDelegates
 
-static bool VerifyModuleCanContainGameplayTag(FName ModuleName, FName TagName, const FModuleDescriptor* Module)
+static bool VerifyModuleCanContainGameplayTag(FName ModuleName, FName TagName, const FModuleDescriptor* Module, TSharedPtr<IPlugin> OptionalOwnerPlugin)
 {
 	if (Module)
 	{
-		if (Module->Type == EHostType::ServerOnly || Module->Type == EHostType::ClientOnly || Module->Type == EHostType::ClientOnlyNoCommandlet)
+		if (!(Module->Type == EHostType::Runtime || Module->Type == EHostType::RuntimeAndProgram))
 		{
-			ensureAlwaysMsgf(false, TEXT("Native Gameplay Tag '%s' defined in '%s', which is Client or Server only module.  Client and Server tags must match."), *TagName.ToString(), *ModuleName.ToString());
+			// TODO NDarnell - If it's not a module we load always we need to make sure the tag is available in some other fashion
+			// such as through an ini.
+
+			//TSharedPtr<IPlugin> ThisPlugin = IPluginManager::Get().FindPlugin(TEXT("CommonDialogue"));
+			//check(ThisPlugin.IsValid());
+			//UGameplayTagsManager::Get().AddTagIniSearchPath(ThisPlugin->GetBaseDir() / TEXT("Config") / TEXT("Tags"));
+
+			//const FString PluginName = FPaths::GetBaseFilename(StateProperties.PluginInstalledFilename);
+			//const FString PluginFolder = FPaths::GetPath(StateProperties.PluginInstalledFilename);
+			//UGameplayTagsManager::Get().AddTagIniSearchPath(PluginFolder / TEXT("Config") / TEXT("Tags"));
+
+
+			//const FString& PluginDescriptorFilename = Plugin->GetDescriptorFileName();
+
+			//// Make sure you are in the game feature plugins folder. All GameFeaturePlugins are in this folder.
+			////@TODO: GameFeaturePluginEnginePush: Comments elsewhere allow plugins outside of the folder as long as they explicitly opt in, either those are wrong or this check is wrong
+			//if (!PluginDescriptorFilename.IsEmpty() && FPaths::ConvertRelativePathToFull(PluginDescriptorFilename).StartsWith(GetDefault<UGameFeaturesSubsystemSettings>()->BuiltInGameFeaturePluginsFolder) && FPaths::FileExists(PluginDescriptorFilename))
+			//{
+
+			ensureAlwaysMsgf(false, TEXT("Native Gameplay Tag '%s' defined in '%s'.  The module type is '%s' but needs to be 'Runtime' or 'RuntimeAndProgram'.  Client and Server tags must match."), *TagName.ToString(), *ModuleName.ToString(), EHostType::ToString(Module->Type));
 		}
 
-		// Not a mistake - we return true even if it fails the test, the return value is an invalidation we were able to verify that it could or could not.
+		// Not a mistake - we return true even if it fails the test, the return value is a validation we were able to verify that it could or could not.
 		return true;
 	}
 
 	return false;
 }
 
+#endif
+
 FNativeGameplayTag::FNativeGameplayTag(FName PluginName, FName ModuleName, FName TagName, const FString& TagDevComment, ENativeGameplayTagToken)
 {
 #if !UE_BUILD_SHIPPING
@@ -49,7 +61,7 @@ FNativeGameplayTag::FNativeGameplayTag(FName PluginName, FName ModuleName, FName
 	const FModuleDescriptor* ProjectModule =
 		CurrentProject->Modules.FindByPredicate([ModuleName](const FModuleDescriptor& Module) { return Module.Name == ModuleName; });
 
-	if (!VerifyModuleCanContainGameplayTag(ModuleName, TagName, ProjectModule))
+	if (!VerifyModuleCanContainGameplayTag(ModuleName, TagName, ProjectModule, TSharedPtr<IPlugin>()))
 	{
 		const FModuleDescriptor* PluginModule = nullptr;
 
@@ -61,14 +73,17 @@ FNativeGameplayTag::FNativeGameplayTag(FName PluginName, FName ModuleName, FName
 			PluginModule = PluginDescriptor.Modules.FindByPredicate([ModuleName](const FModuleDescriptor& Module) { return Module.Name == ModuleName; });
 		}
 
-		if (!VerifyModuleCanContainGameplayTag(ModuleName, TagName, PluginModule))
+		if (!VerifyModuleCanContainGameplayTag(ModuleName, TagName, PluginModule, Plugin))
 		{
 			ensureAlwaysMsgf(false, TEXT("Unable to find information about module '%s' in plugin '%s'"), *ModuleName.ToString(), *PluginName.ToString());
 		}
 	}
 #endif
 
-	//TODO We could restrict creation to during module loading, at least in editor builds.
+	// TODO NDarnell To try and make sure nobody is using these during non-static init
+	// of the module, we could add an indicator on the module manager indicating
+	// if we're actively loading model and make sure we only run this code during
+	// that point.
 
 	InternalTag = TagName.IsNone() ? FGameplayTag() : FGameplayTag(TagName);
 #if WITH_EDITOR
diff --git a/Engine/Source/Runtime/GameplayTags/Public/NativeGameplayTags.h b/Engine/Source/Runtime/GameplayTags/Public/NativeGameplayTags.h
--- a/Engine/Source/Runtime/GameplayTags/Public/NativeGameplayTags.h
+++ b/Engine/Source/Runtime/GameplayTags/Public/NativeGameplayTags.h
@@ -6,50 +6,29 @@
 #include "GameplayTagContainer.h"
 #include "GameplayTagsManager.h"
 
-#define DECLARE_NATIVE_GAMEPLAY_TAG_SOURCE(ClassType)						\
-public:																		\
-	static const TSharedRef<const ClassType>& Get()							\
-	{																		\
-		static TSharedRef<const ClassType> Tags = MakeShared<ClassType>();	\
-		return Tags;														\
-	}																		\
-																			\
-	virtual void Register()	const override									\
-	{																		\
-		UGameplayTagsManager& Manager = UGameplayTagsManager::Get();		\
-		Manager.AddNativeGameplayTagSource(#ClassType, Get());				\
-	}																		\
-																			\
-	virtual void Unregister() const override								\
-	{																		\
-		UGameplayTagsManager& Manager = UGameplayTagsManager::Get();		\
-		Manager.RemoveNativeGameplayTagSource(#ClassType);					\
-	}
-
-class GAMEPLAYTAGS_API FNativeGameplayTagSource
-{
-public:
-	FNativeGameplayTagSource() { }
-	virtual ~FNativeGameplayTagSource() { }
-
-	virtual void Register() const = 0;
-	virtual void Unregister() const = 0;
-
-protected:
-	FGameplayTag Add(FName TagName, const FString& TagDevComment = TEXT("(Native)"));
-
-private:
-	TArray<FGameplayTagTableRow> NativeTags;
-
-	friend class UGameplayTagsManager;
-};
-
 enum class ENativeGameplayTagToken { PRIVATE_USE_MACRO_INSTEAD };
 
+/**
+ * Declares a native gameplay tag that is defined in a cpp with UE_DEFINE_GAMEPLAY_TAG to allow other modules or code to use the created tag variable.
+ */
 #define UE_DECLARE_GAMEPLAY_TAG_EXTERN(TagName) extern FNativeGameplayTag TagName;
+
+/**
+ * Defines a native gameplay tag that is externally declared in a header to allow other modules or code to use the created tag variable.
+ */
 #define UE_DEFINE_GAMEPLAY_TAG(TagName, Tag) FNativeGameplayTag TagName(UE_PLUGIN_NAME, UE_MODULE_NAME, Tag, TEXT(""), ENativeGameplayTagToken::PRIVATE_USE_MACRO_INSTEAD);
+
+/**
+ * Defines a native gameplay tag such that it's only available to the cpp file you define it in.
+ */
 #define UE_DEFINE_GAMEPLAY_TAG_STATIC(TagName, Tag) static FNativeGameplayTag TagName(UE_PLUGIN_NAME, UE_MODULE_NAME, Tag, TEXT(""), ENativeGameplayTagToken::PRIVATE_USE_MACRO_INSTEAD);
 
+/**
+ * Holds a gameplay tag that was registered during static construction of the module, and will
+ * be unregistered when the module unloads.  Each registration is based on the native tag pointer
+ * so even if two modules register the same tag and one is unloaded, the tag will still be registered
+ * by the other one.
+ */
 class GAMEPLAYTAGS_API FNativeGameplayTag : public FNoncopyable
 {
 public:
