Change tag pins to use the same warning and error handling logic as normal serialized tags, they will warn but not clear
Add ClearInvalidTags option to tag settings that if enabled will clear invalid tags on load in the editor for both properties and pins
Change pin widgets to use the standard import text wrappers instead of an error-prone manual parser, and fix normal import to remove empty tags from a container
This change preserves the current behavior of always checking for validity of all non-redirected tags, we may want to change that later

diff --git a/Engine/Plugins/Editor/GameplayTagsEditor/Source/GameplayTagsEditor/Private/SGameplayTagContainerGraphPin.cpp b/Engine/Plugins/Editor/GameplayTagsEditor/Source/GameplayTagsEditor/Private/SGameplayTagContainerGraphPin.cpp
--- a/Engine/Plugins/Editor/GameplayTagsEditor/Source/GameplayTagsEditor/Private/SGameplayTagContainerGraphPin.cpp
+++ b/Engine/Plugins/Editor/GameplayTagsEditor/Source/GameplayTagsEditor/Private/SGameplayTagContainerGraphPin.cpp
@@ -12,60 +12,10 @@ void SGameplayTagContainerGraphPin::Construct(const FArguments& InArgs, UEdGraph
 
 void SGameplayTagContainerGraphPin::ParseDefaultValueData()
 {
-	FString TagString = GraphPinObj->GetDefaultAsString();
-
 	FilterString = GameplayTagPinUtilities::ExtractTagFilterStringFromGraphPin(GraphPinObj);
 
-	// This parsing code should be the same as ImportText but there may be older data this handles better
-	if (TagString.StartsWith(TEXT("("), ESearchCase::CaseSensitive) && TagString.EndsWith(TEXT(")"), ESearchCase::CaseSensitive))
-	{
-		TagString.LeftChopInline(1, false);
-		TagString.RightChopInline(1, false);
-
-		TagString.Split(TEXT("="), nullptr, &TagString, ESearchCase::CaseSensitive);
-
-		TagString.LeftChopInline(1, false);
-		TagString.RightChopInline(1, false);
-
-		FString ReadTag;
-		FString Remainder;
-
-		while (TagString.Split(TEXT(","), &ReadTag, &Remainder, ESearchCase::CaseSensitive))
-		{
-			ReadTag.Split(TEXT("="), NULL, &ReadTag, ESearchCase::CaseSensitive);
-			if (ReadTag.EndsWith(TEXT(")"), ESearchCase::CaseSensitive))
-			{
-				ReadTag.LeftChopInline(1, false);
-				if (ReadTag.StartsWith(TEXT("\""), ESearchCase::CaseSensitive) && ReadTag.EndsWith(TEXT("\""), ESearchCase::CaseSensitive))
-				{
-					ReadTag.LeftChopInline(1, false);
-					ReadTag.RightChopInline(1, false);
-				}
-			}
-			TagString = Remainder;
-			FGameplayTag GameplayTag = FGameplayTag::RequestGameplayTag(FName(*ReadTag));
-			TagContainer->AddTag(GameplayTag);
-		}
-		if (Remainder.IsEmpty())
-		{
-			Remainder = TagString;
-		}
-		if (!Remainder.IsEmpty())
-		{
-			Remainder.Split(TEXT("="), nullptr, &Remainder, ESearchCase::CaseSensitive);
-			if (Remainder.EndsWith(TEXT(")"), ESearchCase::CaseSensitive))
-			{
-				Remainder.LeftChopInline(1, false);
-				if (Remainder.StartsWith(TEXT("\""), ESearchCase::CaseSensitive) && Remainder.EndsWith(TEXT("\""), ESearchCase::CaseSensitive))
-				{
-					Remainder.LeftChopInline(1, false);
-					Remainder.RightChopInline(1, false);
-				}
-			}
-			FGameplayTag GameplayTag = FGameplayTag::RequestGameplayTag(FName(*Remainder));
-			TagContainer->AddTag(GameplayTag);
-		}
-	}
+	// Read using import text, but with serialize flag set so it doesn't always throw away invalid ones
+	TagContainer->FromExportString(GraphPinObj->GetDefaultAsString(), PPF_SerializedAsImportText);
 }
 
 TSharedRef<SWidget> SGameplayTagContainerGraphPin::GetEditContent()
diff --git a/Engine/Plugins/Editor/GameplayTagsEditor/Source/GameplayTagsEditor/Private/SGameplayTagGraphPin.cpp b/Engine/Plugins/Editor/GameplayTagsEditor/Source/GameplayTagsEditor/Private/SGameplayTagGraphPin.cpp
--- a/Engine/Plugins/Editor/GameplayTagsEditor/Source/GameplayTagsEditor/Private/SGameplayTagGraphPin.cpp
+++ b/Engine/Plugins/Editor/GameplayTagsEditor/Source/GameplayTagsEditor/Private/SGameplayTagGraphPin.cpp
@@ -13,25 +13,15 @@ void SGameplayTagGraphPin::Construct( const FArguments& InArgs, UEdGraphPin* InG
 
 void SGameplayTagGraphPin::ParseDefaultValueData()
 {
-	FString TagString = GraphPinObj->GetDefaultAsString();
-
 	FilterString = GameplayTagPinUtilities::ExtractTagFilterStringFromGraphPin(GraphPinObj);
 
-	if (TagString.StartsWith(TEXT("("), ESearchCase::CaseSensitive) && TagString.EndsWith(TEXT(")"), ESearchCase::CaseSensitive))
-	{
-		TagString.LeftChopInline(1, false);
-		TagString.RightChopInline(1, false);
-		TagString.Split(TEXT("="), nullptr, &TagString, ESearchCase::CaseSensitive);
-		if (TagString.StartsWith(TEXT("\""), ESearchCase::CaseSensitive) && TagString.EndsWith(TEXT("\""), ESearchCase::CaseSensitive))
-		{
-			TagString.LeftChopInline(1, false);
-			TagString.RightChopInline(1, false);
-		}
-	}
+	FGameplayTag GameplayTag;
+	
+	// Read using import text, but with serialize flag set so it doesn't always throw away invalid ones
+	GameplayTag.FromExportString(GraphPinObj->GetDefaultAsString(), PPF_SerializedAsImportText);
 
-	if (!TagString.IsEmpty())
+	if (GameplayTag.IsValid())
 	{
-		FGameplayTag GameplayTag = FGameplayTag::RequestGameplayTag(FName(*TagString));
 		TagContainer->AddTag(GameplayTag);
 	}
 }
diff --git a/Engine/Source/Runtime/GameplayTags/Classes/GameplayTagContainer.h b/Engine/Source/Runtime/GameplayTags/Classes/GameplayTagContainer.h
--- a/Engine/Source/Runtime/GameplayTags/Classes/GameplayTagContainer.h
+++ b/Engine/Source/Runtime/GameplayTags/Classes/GameplayTagContainer.h
@@ -191,7 +191,7 @@ struct GAMEPLAYTAGS_API FGameplayTag
 	bool SerializeFromMismatchedTag(const FPropertyTag& Tag, FStructuredArchive::FSlot Slot);
 
 	/** Sets from a ImportText string, used in asset registry */
-	void FromExportString(const FString& ExportString);
+	void FromExportString(const FString& ExportString, int32 PortFlags = 0);
 
 	/** Handles importing tag strings without (TagName=) in it */
 	bool ImportTextItem(const TCHAR*& Buffer, int32 PortFlags, UObject* Parent, FOutputDevice* ErrorText);
@@ -533,7 +533,7 @@ struct GAMEPLAYTAGS_API FGameplayTagContainer
 	FString ToString() const;
 
 	/** Sets from a ImportText string, used in asset registry */
-	void FromExportString(const FString& ExportString);
+	void FromExportString(const FString& ExportString, int32 PortFlags = 0);
 
 	/** Returns abbreviated human readable Tag list without parens or property names. If bQuoted is true it will quote each tag */
 	FString ToStringSimple(bool bQuoted = false) const;
diff --git a/Engine/Source/Runtime/GameplayTags/Classes/GameplayTagsManager.h b/Engine/Source/Runtime/GameplayTags/Classes/GameplayTagsManager.h
--- a/Engine/Source/Runtime/GameplayTags/Classes/GameplayTagsManager.h
+++ b/Engine/Source/Runtime/GameplayTags/Classes/GameplayTagsManager.h
@@ -489,6 +489,12 @@ public:
 		return bShouldWarnOnInvalidTags;
 	}
 
+	/** Should we clear references to invalid tags loaded/saved in the editor */
+	bool ShouldClearInvalidTags() const
+	{
+		return bShouldClearInvalidTags;
+	}
+
 	/** Should use fast replication */
 	bool ShouldUseFastReplication() const
 	{
@@ -520,7 +526,7 @@ public:
 	void RedirectSingleGameplayTag(FGameplayTag& Tag, FProperty* SerializingProperty) const;
 
 	/** Handles establishing a single tag from an imported tag name (accounts for redirects too). Called when tags are imported via text. */
-	bool ImportSingleGameplayTag(FGameplayTag& Tag, FName ImportedTagName) const;
+	bool ImportSingleGameplayTag(FGameplayTag& Tag, FName ImportedTagName, bool bImportFromSerialize = false) const;
 
 	/** Gets a tag name from net index and vice versa, used for replication efficiency */
 	FName GetTagNameFromNetIndex(FGameplayTagNetIndex Index) const;
@@ -763,6 +769,9 @@ private:
 	/** Cached runtime value for whether we should warn when loading invalid tags */
 	bool bShouldWarnOnInvalidTags;
 
+	/** Cached runtime value for whether we should warn when loading invalid tags */
+	bool bShouldClearInvalidTags;
+
 	/** True if native tags have all been added and flushed */
 	bool bDoneAddingNativeTags;
 
diff --git a/Engine/Source/Runtime/GameplayTags/Classes/GameplayTagsSettings.h b/Engine/Source/Runtime/GameplayTags/Classes/GameplayTagsSettings.h
--- a/Engine/Source/Runtime/GameplayTags/Classes/GameplayTagsSettings.h
+++ b/Engine/Source/Runtime/GameplayTags/Classes/GameplayTagsSettings.h
@@ -131,6 +131,10 @@ class GAMEPLAYTAGS_API UGameplayTagsSettings : public UGameplayTagsList
 	UPROPERTY(config, EditAnywhere, Category = GameplayTags, meta = (ConfigRestartRequired = true))
 	bool WarnOnInvalidTags;
 
+	/** If true, will clear any invalid tags when reading in saved tag references that are not in the dictionary */
+	UPROPERTY(config, EditAnywhere, Category = GameplayTags, meta = (ConfigRestartRequired = true))
+	bool ClearInvalidTags;
+
 	/** If true, will replicate gameplay tags by index instead of name. For this to work, tags must be identical on client and server */
 	UPROPERTY(config, EditAnywhere, Category = "Advanced Replication")
 	bool FastReplication;
diff --git a/Engine/Source/Runtime/GameplayTags/Private/GameplayTagContainer.cpp b/Engine/Source/Runtime/GameplayTags/Private/GameplayTagContainer.cpp
--- a/Engine/Source/Runtime/GameplayTags/Private/GameplayTagContainer.cpp
+++ b/Engine/Source/Runtime/GameplayTags/Private/GameplayTagContainer.cpp
@@ -882,12 +882,12 @@ FString FGameplayTagContainer::ToString() const
 	return ExportString;
 }
 
-void FGameplayTagContainer::FromExportString(const FString& ExportString)
+void FGameplayTagContainer::FromExportString(const FString& ExportString, int32 PortFlags)
 {
 	Reset();
 
 	FOutputDeviceNull NullOut;
-	FGameplayTagContainer::StaticStruct()->ImportText(*ExportString, this, nullptr, 0, &NullOut, TEXT("FGameplayTagContainer"), true);
+	FGameplayTagContainer::StaticStruct()->ImportText(*ExportString, this, nullptr, PortFlags, &NullOut, TEXT("FGameplayTagContainer"), true);
 }
 
 bool FGameplayTagContainer::ImportTextItem(const TCHAR*& Buffer, int32 PortFlags, UObject* Parent, FOutputDevice* ErrorText)
@@ -897,6 +897,9 @@ bool FGameplayTagContainer::ImportTextItem(const TCHAR*& Buffer, int32 PortFlags
 
 	if (Buffer)
 	{
+		// Clear out any invalid tags that got stripped
+		GameplayTags.Remove(FGameplayTag());
+
 		// Compute parent tags
 		FillParentTags();	
 	}
@@ -1332,22 +1335,22 @@ bool FGameplayTag::ImportTextItem(const TCHAR*& Buffer, int32 PortFlags, UObject
 
 	if (ImportedTag[0] == '(')
 	{
-		// Let normal ImportText handle this. It appears to be prepared for it.
+		// Let normal ImportText handle this before handling fixups
 		UScriptStruct* ScriptStruct = FGameplayTag::StaticStruct();
 		Buffer = ScriptStruct->ImportText(Buffer, this, Parent, PortFlags, ErrorText, ScriptStruct->GetName(), false);
-		UGameplayTagsManager::Get().ImportSingleGameplayTag(*this, TagName);
+		UGameplayTagsManager::Get().ImportSingleGameplayTag(*this, TagName, !!(PortFlags & PPF_SerializedAsImportText));
 		return true;
 	}
 
-	return UGameplayTagsManager::Get().ImportSingleGameplayTag(*this, FName(*ImportedTag));
+	return UGameplayTagsManager::Get().ImportSingleGameplayTag(*this, FName(*ImportedTag), !!(PortFlags & PPF_SerializedAsImportText));
 }
 
-void FGameplayTag::FromExportString(const FString& ExportString)
+void FGameplayTag::FromExportString(const FString& ExportString, int32 PortFlags)
 {
 	TagName = NAME_None;
 
 	FOutputDeviceNull NullOut;
-	FGameplayTag::StaticStruct()->ImportText(*ExportString, this, nullptr, 0, &NullOut, TEXT("FGameplayTag"), true);
+	FGameplayTag::StaticStruct()->ImportText(*ExportString, this, nullptr, PortFlags, &NullOut, TEXT("FGameplayTag"), true);
 }
 
 FGameplayTagNativeAdder::FGameplayTagNativeAdder()
diff --git a/Engine/Source/Runtime/GameplayTags/Private/GameplayTagsManager.cpp b/Engine/Source/Runtime/GameplayTags/Private/GameplayTagsManager.cpp
--- a/Engine/Source/Runtime/GameplayTags/Private/GameplayTagsManager.cpp
+++ b/Engine/Source/Runtime/GameplayTags/Private/GameplayTagsManager.cpp
@@ -91,6 +91,7 @@ UGameplayTagsManager::UGameplayTagsManager(const FObjectInitializer& ObjectIniti
 {
 	bUseFastReplication = false;
 	bShouldWarnOnInvalidTags = true;
+	bShouldClearInvalidTags = false;
 	bDoneAddingNativeTags = false;
 	NetIndexFirstBitSegment = 16;
 	NetIndexTrueBitNum = 16;
@@ -405,6 +406,7 @@ void UGameplayTagsManager::ConstructGameplayTagTree()
 
 			bUseFastReplication = MutableDefault->FastReplication;
 			bShouldWarnOnInvalidTags = MutableDefault->WarnOnInvalidTags;
+			bShouldClearInvalidTags = MutableDefault->ClearInvalidTags;
 			NumBitsForContainerSize = MutableDefault->NumBitsForContainerSize;
 			NetIndexFirstBitSegment = MutableDefault->NetIndexFirstBitSegment;
 		}
@@ -776,11 +778,19 @@ void UGameplayTagsManager::RedirectTagsForContainer(FGameplayTagContainer& Conta
 		{
 			// Warn about invalid tags at load time in editor builds, too late to fix it in cooked builds
 			FGameplayTag OldTag = RequestGameplayTag(TagName, false);
-			if (!OldTag.IsValid() && ShouldWarnOnInvalidTags())
+			if (!OldTag.IsValid())
 			{
-				FUObjectSerializeContext* LoadContext = FUObjectThreadContext::Get().GetSerializeContext();
-				UObject* LoadingObject = LoadContext ? LoadContext->SerializedObject : nullptr;
-				UE_ASSET_LOG(LogGameplayTags, Warning, *GetPathNameSafe(LoadingObject), TEXT("Invalid GameplayTag %s found in property %s."), *TagName.ToString(), *GetPathNameSafe(SerializingProperty));
+				if (ShouldWarnOnInvalidTags())
+				{
+					FUObjectSerializeContext* LoadContext = FUObjectThreadContext::Get().GetSerializeContext();
+					UObject* LoadingObject = LoadContext ? LoadContext->SerializedObject : nullptr;
+					UE_ASSET_LOG(LogGameplayTags, Warning, *GetPathNameSafe(LoadingObject), TEXT("Invalid GameplayTag %s found in property %s."), *TagName.ToString(), *GetPathNameSafe(SerializingProperty));
+				}
+
+				if (ShouldClearInvalidTags())
+				{
+					NamesToRemove.Add(TagName);
+				}
 			}
 		}
 #endif
@@ -827,21 +837,29 @@ void UGameplayTagsManager::RedirectSingleGameplayTag(FGameplayTag& Tag, FPropert
 		}
 	}
 #if WITH_EDITOR
-	else if (TagName != NAME_None && SerializingProperty)
+	else if (!TagName.IsNone() && SerializingProperty)
 	{
 		// Warn about invalid tags at load time in editor builds, too late to fix it in cooked builds
 		FGameplayTag OldTag = RequestGameplayTag(TagName, false);
-		if (!OldTag.IsValid() && ShouldWarnOnInvalidTags())
+		if (!OldTag.IsValid())
 		{
-			FUObjectSerializeContext* LoadContext = FUObjectThreadContext::Get().GetSerializeContext();
-			UObject* LoadingObject = LoadContext ? LoadContext->SerializedObject : nullptr;
-			UE_ASSET_LOG(LogGameplayTags, Warning, *GetPathNameSafe(LoadingObject), TEXT("Invalid GameplayTag %s found in property %s."), *TagName.ToString(), *GetPathNameSafe(SerializingProperty));
+			if (ShouldWarnOnInvalidTags())
+			{
+				FUObjectSerializeContext* LoadContext = FUObjectThreadContext::Get().GetSerializeContext();
+				UObject* LoadingObject = LoadContext ? LoadContext->SerializedObject : nullptr;
+				UE_ASSET_LOG(LogGameplayTags, Warning, *GetPathNameSafe(LoadingObject), TEXT("Invalid GameplayTag %s found in property %s."), *TagName.ToString(), *GetPathNameSafe(SerializingProperty));
+			}
+			
+			if (ShouldClearInvalidTags())
+			{
+				Tag.TagName = NAME_None;
+			}
 		}
 	}
 #endif
 }
 
-bool UGameplayTagsManager::ImportSingleGameplayTag(FGameplayTag& Tag, FName ImportedTagName) const
+bool UGameplayTagsManager::ImportSingleGameplayTag(FGameplayTag& Tag, FName ImportedTagName, bool bImportFromSerialize) const
 {
 	bool bRetVal = false;
 	if (const FGameplayTag* RedirectedTag = TagRedirects.Find(ImportedTagName))
@@ -856,6 +874,30 @@ bool UGameplayTagsManager::ImportSingleGameplayTag(FGameplayTag& Tag, FName Impo
 		bRetVal = true;
 	}
 
+	if (!bRetVal && bImportFromSerialize && !ImportedTagName.IsNone())
+	{
+#if WITH_EDITOR
+		if (ShouldWarnOnInvalidTags())
+		{
+			FUObjectSerializeContext* LoadContext = FUObjectThreadContext::Get().GetSerializeContext();
+			UObject* LoadingObject = LoadContext ? LoadContext->SerializedObject : nullptr;
+			if (LoadingObject)
+			{
+				// If this is a serialize with a real object and it failed to find the tag, warn about it
+				UE_ASSET_LOG(LogGameplayTags, Warning, *GetPathNameSafe(LoadingObject), TEXT("Invalid GameplayTag %s found in object %s."), *ImportedTagName.ToString(), *LoadingObject->GetName());
+			}
+		}
+
+		// Always keep invalid tags in cooked game to be consistent with properties
+		if (!ShouldClearInvalidTags())
+#endif
+		{
+			// For imported tags that are part of a serialize, leave invalid ones the same way normal serialization does to avoid data loss
+			Tag.TagName = ImportedTagName;
+			bRetVal = true;
+		}
+	}
+
 	if (bRetVal)
 	{
 		OnGameplayTagLoadedDelegate.Broadcast(Tag);
@@ -1721,7 +1763,7 @@ FGameplayTag UGameplayTagsManager::RequestGameplayTag(FName TagName, bool ErrorI
 		static TSet<FName> MissingTagName;
 		if (!MissingTagName.Contains(TagName))
 		{
-			ensureAlwaysMsgf(false, TEXT("Requested Tag %s was not found. Check tag data table."), *TagName.ToString());
+			ensureAlwaysMsgf(false, TEXT("Requested Gameplay Tag %s was not found, tags must be loaded from config or registered as a native tag"), *TagName.ToString());
 			MissingTagName.Add(TagName);
 		}
 	}
diff --git a/Engine/Source/Runtime/GameplayTags/Private/GameplayTagsSettings.cpp b/Engine/Source/Runtime/GameplayTags/Private/GameplayTagsSettings.cpp
--- a/Engine/Source/Runtime/GameplayTags/Private/GameplayTagsSettings.cpp
+++ b/Engine/Source/Runtime/GameplayTags/Private/GameplayTagsSettings.cpp
@@ -61,6 +61,7 @@ UGameplayTagsSettings::UGameplayTagsSettings(const FObjectInitializer& ObjectIni
 	ConfigFileName = GetDefaultConfigFilename();
 	ImportTagsFromConfig = true;
 	WarnOnInvalidTags = true;
+	ClearInvalidTags = false;
 	FastReplication = false;
 	InvalidTagCharacters = ("\"',");
 	NumBitsForContainerSize = 6;
