Keep track of the list of registered gameplay tag search root paths and allow selecting from them when making a new tag source
Add a tooltip with the tag source's actual config file to the add tag screen

diff --git a/Engine/Plugins/Editor/GameplayTagsEditor/Source/GameplayTagsEditor/Private/SAddNewGameplayTagSourceWidget.cpp b/Engine/Plugins/Editor/GameplayTagsEditor/Source/GameplayTagsEditor/Private/SAddNewGameplayTagSourceWidget.cpp
--- a/Engine/Plugins/Editor/GameplayTagsEditor/Source/GameplayTagsEditor/Private/SAddNewGameplayTagSourceWidget.cpp
+++ b/Engine/Plugins/Editor/GameplayTagsEditor/Source/GameplayTagsEditor/Private/SAddNewGameplayTagSourceWidget.cpp
@@ -21,6 +21,7 @@ void SAddNewGameplayTagSourceWidget::Construct(const FArguments& InArgs)
 	bShouldGetKeyboardFocus = false;
 
 	OnGameplayTagSourceAdded = InArgs._OnGameplayTagSourceAdded;
+	PopulateTagRoots();
 
 	ChildSlot
 	[
@@ -51,6 +52,39 @@ void SAddNewGameplayTagSourceWidget::Construct(const FArguments& InArgs)
 			]
 		]
 
+		// Tag source root
+		+ SVerticalBox::Slot()
+		.AutoHeight()
+		.VAlign(VAlign_Top)
+		[
+			SNew(SHorizontalBox)
+			+ SHorizontalBox::Slot()
+			.Padding(2.0f, 6.0f)
+			.AutoWidth()
+			[
+				SNew(STextBlock)
+				.Text(LOCTEXT("ConfigPathLabel", "Config Path:"))
+				.ToolTipText(LOCTEXT("RootPathTooltip", "Set the base config path for added source, this includes paths from plugins and other places that call AddTagIniSearchPath"))
+			]
+
+			+ SHorizontalBox::Slot()
+			.Padding(2.0f, 2.0f)
+			.FillWidth(1.0f)
+			.HAlign(HAlign_Right)
+			[
+				SAssignNew(TagRootsComboBox, SComboBox<TSharedPtr<FString> >)
+				.OptionsSource(&TagRoots)
+				.OnGenerateWidget(this, &SAddNewGameplayTagSourceWidget::OnGenerateTagRootsComboBox)
+				.ContentPadding(2.0f)
+				.Content()
+				[
+					SNew(STextBlock)
+					.Text(this, &SAddNewGameplayTagSourceWidget::CreateTagRootsComboBoxContent)
+					.Font(IDetailLayoutBuilder::GetDetailFont())
+				]
+			]
+		]
+
 		// Add Source Button
 		+SVerticalBox::Slot()
 		.AutoHeight()
@@ -91,13 +125,70 @@ void SAddNewGameplayTagSourceWidget::SetSourceName(const FText& InName)
 	SourceNameTextBox->SetText(InName.IsEmpty() ? FText::FromString(DefaultNewName) : InName);
 }
 
+void SAddNewGameplayTagSourceWidget::PopulateTagRoots()
+{
+	UGameplayTagsManager& Manager = UGameplayTagsManager::Get();
+	TagRoots.Empty();
+
+	FName DefaultSource = FGameplayTagSource::GetDefaultName();
+
+	TArray<FString> TagRootStrings;
+	Manager.GetTagSourceSearchPaths(TagRootStrings);
+
+	for (const FString& TagRoot : TagRootStrings)
+	{
+		TagRoots.Add(MakeShareable(new FString(TagRoot)));
+	}
+}
+
+FText SAddNewGameplayTagSourceWidget::GetFriendlyPath(TSharedPtr<FString> InItem) const
+{
+	if (InItem.IsValid())
+	{
+		FString FilePath = *InItem.Get();
+
+		if (FPaths::IsUnderDirectory(FilePath, FPaths::ProjectDir()))
+		{
+			FPaths::MakePathRelativeTo(FilePath, *FPaths::ProjectDir());
+		}
+		return FText::FromString(FilePath);
+	}
+	return FText();
+}
+
+TSharedRef<SWidget> SAddNewGameplayTagSourceWidget::OnGenerateTagRootsComboBox(TSharedPtr<FString> InItem)
+{
+	return SNew(STextBlock)
+		.Text(GetFriendlyPath(InItem));
+}
+
+FText SAddNewGameplayTagSourceWidget::CreateTagRootsComboBoxContent() const
+{
+	const bool bHasSelectedItem = TagRootsComboBox.IsValid() && TagRootsComboBox->GetSelectedItem().IsValid();
+
+	if (bHasSelectedItem)
+	{
+		return GetFriendlyPath(TagRootsComboBox->GetSelectedItem());
+	}
+	else
+	{
+		return LOCTEXT("NewTagRootNotSelected", "Default");
+	}
+}
+
 FReply SAddNewGameplayTagSourceWidget::OnAddNewSourceButtonPressed()
 {
 	UGameplayTagsManager& Manager = UGameplayTagsManager::Get();
 		
 	if (!SourceNameTextBox->GetText().EqualTo(FText::FromString(DefaultNewName)))
 	{
-		Manager.FindOrAddTagSource(*SourceNameTextBox->GetText().ToString(), EGameplayTagSourceType::TagList);
+		FString TagRoot;
+		if (TagRootsComboBox->GetSelectedItem().Get())
+		{
+			TagRoot = *TagRootsComboBox->GetSelectedItem().Get();
+		}
+
+		Manager.FindOrAddTagSource(*SourceNameTextBox->GetText().ToString(), EGameplayTagSourceType::TagList, TagRoot);
 	}
 
 	IGameplayTagsModule::OnTagSettingsChanged.Broadcast();
diff --git a/Engine/Plugins/Editor/GameplayTagsEditor/Source/GameplayTagsEditor/Private/SAddNewGameplayTagSourceWidget.h b/Engine/Plugins/Editor/GameplayTagsEditor/Source/GameplayTagsEditor/Private/SAddNewGameplayTagSourceWidget.h
--- a/Engine/Plugins/Editor/GameplayTagsEditor/Source/GameplayTagsEditor/Private/SAddNewGameplayTagSourceWidget.h
+++ b/Engine/Plugins/Editor/GameplayTagsEditor/Source/GameplayTagsEditor/Private/SAddNewGameplayTagSourceWidget.h
@@ -37,17 +37,33 @@ private:
 	/** Sets the name of the source. Uses the default if the name is not specified */
 	void SetSourceName(const FText& InName = FText());
 
+	/** Creates a list of all root directories for tag sources */
+	void PopulateTagRoots();
+
+	/** Populates the widget's combo box with all potential places where a tag source can be stored */
+	TSharedRef<SWidget> OnGenerateTagRootsComboBox(TSharedPtr<FString> InItem);
+
+	/** Gets friendly version of a root path */
+	FText GetFriendlyPath(TSharedPtr<FString> InItem) const;
+
+	/** Creates the text displayed by the combo box when an option is selected */
+	FText CreateTagRootsComboBoxContent() const;
+
 	/** Callback for when the Add New Tag button is pressed */
 	FReply OnAddNewSourceButtonPressed();
 
-private:
-
 	/** The name of the next gameplay tag source to create */
 	TSharedPtr<SEditableTextBox> SourceNameTextBox;
 
 	/** Callback for when a new gameplay tag has been added to the INI files */
 	FOnGameplayTagSourceAdded OnGameplayTagSourceAdded;
 
+	/** All potential INI files where a tag source can be stored */
+	TArray<TSharedPtr<FString> > TagRoots;
+
+	/** The directory where the next tag source will be created */
+	TSharedPtr<SComboBox<TSharedPtr<FString> > > TagRootsComboBox;
+
 	/** Tracks if this widget should get keyboard focus */
 	bool bShouldGetKeyboardFocus;
 
diff --git a/Engine/Plugins/Editor/GameplayTagsEditor/Source/GameplayTagsEditor/Private/SAddNewGameplayTagWidget.cpp b/Engine/Plugins/Editor/GameplayTagsEditor/Source/GameplayTagsEditor/Private/SAddNewGameplayTagWidget.cpp
--- a/Engine/Plugins/Editor/GameplayTagsEditor/Source/GameplayTagsEditor/Private/SAddNewGameplayTagWidget.cpp
+++ b/Engine/Plugins/Editor/GameplayTagsEditor/Source/GameplayTagsEditor/Private/SAddNewGameplayTagWidget.cpp
@@ -117,6 +117,7 @@ void SAddNewGameplayTagWidget::Construct(const FArguments& InArgs)
 				SAssignNew(TagSourcesComboBox, SComboBox<TSharedPtr<FName> >)
 				.OptionsSource(&TagSources)
 				.OnGenerateWidget(this, &SAddNewGameplayTagWidget::OnGenerateTagSourcesComboBox)
+				.ToolTipText(this, &SAddNewGameplayTagWidget::CreateTagSourcesComboBoxToolTip)
 				.ContentPadding(2.0f)
 				.Content()
 				[
@@ -345,4 +346,27 @@ FText SAddNewGameplayTagWidget::CreateTagSourcesComboBoxContent() const
 	return bHasSelectedItem ? FText::FromName(*TagSourcesComboBox->GetSelectedItem().Get()) : LOCTEXT("NewTagLocationNotSelected", "Not selected");
 }
 
+FText SAddNewGameplayTagWidget::CreateTagSourcesComboBoxToolTip() const
+{
+	const bool bHasSelectedItem = TagSourcesComboBox.IsValid() && TagSourcesComboBox->GetSelectedItem().IsValid();
+
+	if (bHasSelectedItem)
+	{
+		UGameplayTagsManager& Manager = UGameplayTagsManager::Get();
+		const FGameplayTagSource* Source = Manager.FindTagSource(*TagSourcesComboBox->GetSelectedItem().Get());
+		if (Source)
+		{
+			FString FilePath = Source->GetConfigFileName();
+
+			if (FPaths::IsUnderDirectory(FilePath, FPaths::ProjectDir()))
+			{
+				FPaths::MakePathRelativeTo(FilePath, *FPaths::ProjectDir());
+			}
+			return FText::FromString(FilePath);
+		}
+	}
+
+	return FText();
+}
+
 #undef LOCTEXT_NAMESPACE
diff --git a/Engine/Plugins/Editor/GameplayTagsEditor/Source/GameplayTagsEditor/Private/SAddNewGameplayTagWidget.h b/Engine/Plugins/Editor/GameplayTagsEditor/Source/GameplayTagsEditor/Private/SAddNewGameplayTagWidget.h
--- a/Engine/Plugins/Editor/GameplayTagsEditor/Source/GameplayTagsEditor/Private/SAddNewGameplayTagWidget.h
+++ b/Engine/Plugins/Editor/GameplayTagsEditor/Source/GameplayTagsEditor/Private/SAddNewGameplayTagWidget.h
@@ -78,6 +78,9 @@ private:
 	/** Creates the text displayed by the combo box when an option is selected */
 	FText CreateTagSourcesComboBoxContent() const;
 
+	/** Creates the text displayed by the combo box tooltip when an option is selected */
+	FText CreateTagSourcesComboBoxToolTip() const;
+
 	EVisibility OnGetTagSourceFavoritesVisibility() const;
 	FReply OnToggleTagSourceFavoriteClicked();
 	const FSlateBrush* OnGetTagSourceFavoriteImage() const;
diff --git a/Engine/Source/Runtime/GameplayTags/Classes/GameplayTagsManager.h b/Engine/Source/Runtime/GameplayTags/Classes/GameplayTagsManager.h
--- a/Engine/Source/Runtime/GameplayTags/Classes/GameplayTagsManager.h
+++ b/Engine/Source/Runtime/GameplayTags/Classes/GameplayTagsManager.h
@@ -116,6 +116,9 @@ struct GAMEPLAYTAGS_API FGameplayTagSource
 	{
 	}
 
+	/** Returns the config file that created this source, if valid */
+	FString GetConfigFileName() const;
+
 	static FName GetNativeName();
 
 	static FName GetDefaultName();
@@ -602,6 +605,9 @@ public:
 	/** Refresh the gameplaytag tree due to an editor change */
 	void EditorRefreshGameplayTagTree();
 
+	/** Gets all the current directories to look for tag sources in */
+	void GetTagSourceSearchPaths(TArray<FString>& OutPaths);
+
 	/** Gets a Tag Container containing all of the tags in the hierarchy that are children of this tag, and were explicitly added to the dictionary */
 	FGameplayTagContainer RequestGameplayTagChildrenInDictionary(const FGameplayTag& GameplayTag) const;
 #if WITH_EDITORONLY_DATA
@@ -718,7 +724,7 @@ private:
 	void GetAllParentNodeNames(TSet<FName>& NamesList, TSharedPtr<FGameplayTagNode> GameplayTag) const;
 
 	/** Returns the tag source for a given tag source name, or null if not found */
-	FGameplayTagSource* FindOrAddTagSource(FName TagSourceName, EGameplayTagSourceType SourceType);
+	FGameplayTagSource* FindOrAddTagSource(FName TagSourceName, EGameplayTagSourceType SourceType, const FString& RootDirToUse = FString());
 
 	/** Constructs the net indices for each tag */
 	void ConstructNetIndex();
@@ -759,6 +765,7 @@ private:
 
 	TSet<FName> RestrictedGameplayTagSourceNames;
 
+	/** Specific list of found ini files to load tags out of */
 	TArray<FString> ExtraTagIniList;
 
 	bool bIsConstructingGameplayTagTree = false;
@@ -785,6 +792,9 @@ private:
 
 	// Transient editor-only tags to support quick-iteration PIE workflows
 	TSet<FName> TransientEditorTags;
+
+	/** List of additional tag config root directories to search */
+	TArray<FString> TagIniSearchPaths;
 #endif
 
 	/** Sorted list of nodes, used for network replication */
diff --git a/Engine/Source/Runtime/GameplayTags/Private/GameplayTagsManager.cpp b/Engine/Source/Runtime/GameplayTags/Private/GameplayTagsManager.cpp
--- a/Engine/Source/Runtime/GameplayTags/Private/GameplayTagsManager.cpp
+++ b/Engine/Source/Runtime/GameplayTags/Private/GameplayTagsManager.cpp
@@ -41,6 +41,20 @@ const FName UGameplayTagsManager::NAME_GameplayTagFilter("GameplayTagFilter");
 static const FName NAME_Native = FName(TEXT("Native"));
 static const FName NAME_DefaultGameplayTagsIni("DefaultGameplayTags.ini");
 
+FString FGameplayTagSource::GetConfigFileName() const
+{
+	if (SourceTagList)
+	{
+		return SourceTagList->ConfigFileName;
+	}
+	if (SourceRestrictedTagList)
+	{
+		return SourceRestrictedTagList->ConfigFileName;
+	}
+
+	return FString();
+}
+
 FName FGameplayTagSource::GetNativeName()
 {
 	return NAME_Native;
@@ -159,6 +173,10 @@ struct FCompareFGameplayTagNodeByTag
 
 void UGameplayTagsManager::AddTagIniSearchPath(const FString& RootDir)
 {
+#if WITH_EDITOR
+	TagIniSearchPaths.AddUnique(RootDir);
+#endif
+
 	// Read all tags from the ini
 	TArray<FString> FilesInDirectory;
 	IFileManager::Get().FindFilesRecursive(FilesInDirectory, *RootDir, TEXT("*.ini"), true, false);
@@ -1585,6 +1603,11 @@ void UGameplayTagsManager::EditorRefreshGameplayTagTree()
 	OnEditorRefreshGameplayTagTree.Broadcast();
 }
 
+void UGameplayTagsManager::GetTagSourceSearchPaths(TArray<FString>& OutPaths)
+{
+	OutPaths = TagIniSearchPaths;
+}
+
 FGameplayTagContainer UGameplayTagsManager::RequestGameplayTagChildrenInDictionary(const FGameplayTag& GameplayTag) const
 {
 	// Note this purposefully does not include the passed in GameplayTag in the container.
@@ -1686,7 +1709,7 @@ void UGameplayTagsManager::FindTagSourcesWithType(EGameplayTagSourceType TagSour
 	}
 }
 
-FGameplayTagSource* UGameplayTagsManager::FindOrAddTagSource(FName TagSourceName, EGameplayTagSourceType SourceType)
+FGameplayTagSource* UGameplayTagsManager::FindOrAddTagSource(FName TagSourceName, EGameplayTagSourceType SourceType, const FString& RootDirToUse)
 {
 	FGameplayTagSource* FoundSource = FindTagSource(TagSourceName);
 	if (FoundSource)
@@ -1710,7 +1733,16 @@ FGameplayTagSource* UGameplayTagsManager::FindOrAddTagSource(FName TagSourceName
 	else if (SourceType == EGameplayTagSourceType::TagList)
 	{
 		NewSource->SourceTagList = NewObject<UGameplayTagsList>(this, TagSourceName, RF_Transient);
-		NewSource->SourceTagList->ConfigFileName = FString::Printf(TEXT("%sTags/%s"), *FPaths::SourceConfigDir(), *TagSourceName.ToString());
+		if (RootDirToUse.IsEmpty())
+		{
+			NewSource->SourceTagList->ConfigFileName = FString::Printf(TEXT("%sTags/%s"), *FPaths::SourceConfigDir(), *TagSourceName.ToString());
+		}
+		else
+		{
+			// Use custom root and make sure it gets added to the ini list for later refresh
+			NewSource->SourceTagList->ConfigFileName = RootDirToUse / *TagSourceName.ToString();
+			ExtraTagIniList.AddUnique(NewSource->SourceTagList->ConfigFileName);
+		}
 		if (GUObjectArray.IsDisregardForGC(this))
 		{
 			NewSource->SourceTagList->AddToRoot();
@@ -1719,7 +1751,16 @@ FGameplayTagSource* UGameplayTagsManager::FindOrAddTagSource(FName TagSourceName
 	else if (SourceType == EGameplayTagSourceType::RestrictedTagList)
 	{
 		NewSource->SourceRestrictedTagList = NewObject<URestrictedGameplayTagsList>(this, TagSourceName, RF_Transient);
-		NewSource->SourceRestrictedTagList->ConfigFileName = FString::Printf(TEXT("%sTags/%s"), *FPaths::SourceConfigDir(), *TagSourceName.ToString());
+		if (RootDirToUse.IsEmpty())
+		{
+			NewSource->SourceRestrictedTagList->ConfigFileName = FString::Printf(TEXT("%sTags/%s"), *FPaths::SourceConfigDir(), *TagSourceName.ToString());
+		}
+		else
+		{
+			// Use custom root and make sure it gets added to the ini list for later refresh
+			NewSource->SourceRestrictedTagList->ConfigFileName = RootDirToUse / *TagSourceName.ToString();
+			ExtraTagIniList.AddUnique(NewSource->SourceTagList->ConfigFileName);
+		}
 		if (GUObjectArray.IsDisregardForGC(this))
 		{
 			NewSource->SourceRestrictedTagList->AddToRoot();
