Remove threadsafe hack that was added to one tag function as the tag node access is threadsafe as of 4.27/5.0 and this fix only caught one of the many execution flows that could cause the same shared ptr reference issues
The core issue was addressed for UE-113567, if we need further fixes it should happen in GetSingleTagContainer instead of here

diff --git a/Engine/Source/Runtime/GameplayTags/Private/GameplayTagContainer.cpp b/Engine/Source/Runtime/GameplayTags/Private/GameplayTagContainer.cpp
--- a/Engine/Source/Runtime/GameplayTags/Private/GameplayTagContainer.cpp
+++ b/Engine/Source/Runtime/GameplayTags/Private/GameplayTagContainer.cpp
@@ -596,30 +596,13 @@ bool FGameplayTagContainer::RemoveTagByExplicitName(const FName& TagName)
 
 FORCEINLINE_DEBUGGABLE void FGameplayTagContainer::AddParentsForTag(const FGameplayTag& Tag)
 {
-	if (IsInGameThread())
+	const FGameplayTagContainer* SingleContainer = UGameplayTagsManager::Get().GetSingleTagContainer(Tag);
+	if (SingleContainer)
 	{
-		const FGameplayTagContainer* SingleContainer = UGameplayTagsManager::Get().GetSingleTagContainer(Tag);
-
-		if (SingleContainer)
-		{
-			// Add Parent tags from this tag to our own
-			for (const FGameplayTag& ParentTag : SingleContainer->ParentTags)
-			{
-				ParentTags.AddUnique(ParentTag);
-			}
-		}
-	}
-	else
-	{
-		// Non-main threads may not access FGameplayTagNodes, so we can't call GetSingleTagContainer.
-		// This string parsing is probably slower, but threadsafe
-		FString TagString = *Tag.ToString();
-		for(int32 CharIdx = TagString.Len() - 1; CharIdx > 0; --CharIdx)
+		// Add Parent tags from this tag to our own
+		for (const FGameplayTag& ParentTag : SingleContainer->ParentTags)
 		{
-			if (TagString[CharIdx] == '.')
-			{
-				ParentTags.AddUnique(FGameplayTag(FName(*TagString.Left(CharIdx))));
-			}
+			ParentTags.AddUnique(ParentTag);
 		}
 	}
 }
