[NO NON-PII COMMIT MESSAGE]

diff --git a/Engine/Source/Runtime/GameplayTags/Classes/GameplayTagsManager.h b/Engine/Source/Runtime/GameplayTags/Classes/GameplayTagsManager.h
--- a/Engine/Source/Runtime/GameplayTags/Classes/GameplayTagsManager.h
+++ b/Engine/Source/Runtime/GameplayTags/Classes/GameplayTagsManager.h
@@ -401,11 +401,29 @@ public:
 	 */
 	FORCEINLINE_DEBUGGABLE const FGameplayTagContainer* GetSingleTagContainer(const FGameplayTag& GameplayTag) const
 	{
-		TSharedPtr<FGameplayTagNode> TagNode = FindTagNode(GameplayTag);
-		if (TagNode.IsValid())
+		// Doing this with pointers to avoid a shared ptr reference count change
+		const TSharedPtr<FGameplayTagNode>* Node = GameplayTagNodeMap.Find(GameplayTag);
+
+		if (Node)
 		{
-			return &(TagNode->GetSingleTagContainer());
+			return &(*Node)->GetSingleTagContainer();
+		}
+#if WITH_EDITOR
+		// Check redirector
+		if (GIsEditor && GameplayTag.IsValid())
+		{
+			FGameplayTag RedirectedTag = GameplayTag;
+
+			RedirectSingleGameplayTag(RedirectedTag, nullptr);
+
+			Node = GameplayTagNodeMap.Find(RedirectedTag);
+
+			if (Node)
+			{
+				return &(*Node)->GetSingleTagContainer();
+			}
 		}
+#endif
 		return nullptr;
 	}
 
