Adding GameplayTag filter support  to keys of an TMap
 + usage Meta = (Categories = <FilterTag>)

diff --git a/Engine/Source/Runtime/GameplayTags/Private/GameplayTagsManager.cpp b/Engine/Source/Runtime/GameplayTags/Private/GameplayTagsManager.cpp
--- a/Engine/Source/Runtime/GameplayTags/Private/GameplayTagsManager.cpp
+++ b/Engine/Source/Runtime/GameplayTags/Private/GameplayTagsManager.cpp
@@ -1482,6 +1482,15 @@ FString UGameplayTagsManager::StaticGetCategoriesMetaFromPropertyHandle(TSharedP
 					break;
 				}
 			}
+
+			/**	TMap<FGameplayTag, ValueType> GameplayTagMap; */
+			if (FMapProperty* MapProperty = CastField<FMapProperty>(Property))
+			{
+				if (GetFieldMetaData(MapProperty->KeyProp))
+				{
+					break;
+				}
+			}
 		}
 		PropertyHandle = PropertyHandle->GetParentHandle();
 	}
