Added new function GetGameplayTagArray to FGameplayTagContainer so we can get the explicit list of all unique gameplay tags referenced by the query.

This function allows us to register for events when any of the relevant tags are added/removed so we can update the query result.

#UE5 #RNX #NoReleaseNotes

diff --git a/Engine/Source/Runtime/GameplayTags/Classes/GameplayTagContainer.h b/Engine/Source/Runtime/GameplayTags/Classes/GameplayTagContainer.h
--- a/Engine/Source/Runtime/GameplayTags/Classes/GameplayTagContainer.h
+++ b/Engine/Source/Runtime/GameplayTags/Classes/GameplayTagContainer.h
@@ -877,6 +877,12 @@ public:
 	/** Returns description string. */
 	const FString& GetDescription() const { return UserDescription.IsEmpty() ? AutoDescription : UserDescription; };
 
+	//** Gets the explicit list of all unique gameplay tags referenced by the query. */
+	void GetGameplayTagArray(TArray<FGameplayTag>& OutGameplayTags) const
+	{
+		OutGameplayTags = TagDictionary;
+	}
+
 #if WITH_EDITOR
 	/** Creates this query based on the given EditableQuery object */
 	void BuildFromEditableQuery(class UEditableGameplayTagQuery& EditableQuery); 
