Improve ensure and comments related to using an invalid gameplay tag for container queries

diff --git a/Engine/Source/Runtime/GameplayTags/Private/GameplayTagContainer.cpp b/Engine/Source/Runtime/GameplayTags/Private/GameplayTagContainer.cpp
--- a/Engine/Source/Runtime/GameplayTags/Private/GameplayTagContainer.cpp
+++ b/Engine/Source/Runtime/GameplayTags/Private/GameplayTagContainer.cpp
@@ -1085,8 +1085,9 @@ bool FGameplayTag::MatchesTag(const FGameplayTag& TagToCheck) const
 		return TagContainer->HasTag(TagToCheck);
 	}
 
-	// This should always be invalid if the node is missing
-	ensureMsgf(!IsValid(), TEXT("Valid tag failed to convert to single tag container. %s"), *GetTagName().ToString());
+	// If a non-empty tag has not been registered, it will not exist in the tag database so this function may return the incorrect value
+	// All tags must be registered from code or data before being used in matching functions and this tag may have been deleted with active references
+	ensureMsgf(!IsValid(), TEXT("MatchesTag passed invalid gameplay tag %s, only registered tags can be used in containers"), *GetTagName().ToString());
 
 	return false;
 }
@@ -1102,8 +1103,10 @@ bool FGameplayTag::MatchesAny(const FGameplayTagContainer& ContainerToCheck) con
 		return TagContainer->HasAny(ContainerToCheck);
 	}
 
-	// This should always be invalid if the node is missing
-	ensureMsgf(!IsValid(), TEXT("Valid tag failed to conver to single tag container. %s"), *GetTagName().ToString() );
+	// If a non-empty tag has not been registered, it will not exist in the tag database so this function may return the incorrect value
+	// All tags must be registered from code or data before being used in matching functions and this tag may have been deleted with active references
+	ensureMsgf(!IsValid(), TEXT("MatchesAny passed invalid gameplay tag %s, only registered tags can be used in containers"), *GetTagName().ToString());
+
 	return false;
 }
 
