Make sure GameplayTagsManager uses standard filenames for discovered INIs.
Related to CL-19527807 which enforced standard filenames to prevent duplicate loading / contention among INI files.

diff --git a/Engine/Source/Runtime/GameplayTags/Private/GameplayTagsManager.cpp b/Engine/Source/Runtime/GameplayTags/Private/GameplayTagsManager.cpp
--- a/Engine/Source/Runtime/GameplayTags/Private/GameplayTagsManager.cpp
+++ b/Engine/Source/Runtime/GameplayTags/Private/GameplayTagsManager.cpp
@@ -259,7 +259,7 @@ void UGameplayTagsManager::AddTagIniSearchPath(const FString& RootDir)
 			{
 				const FName TagSource = FName(*FPaths::GetCleanFilename(IniFilePath));
 				PathInfo->SourcesInPath.Add(TagSource);
-				PathInfo->TagIniList.Add(IniFilePath);
+				PathInfo->TagIniList.Add(FPaths::CreateStandardFilename(IniFilePath));
 			}
 		}
 		
