#UE In UGameplayTagsManager, use the default NetIndexFirstBitSegment when clamping to NetIndexTrueBitNum. This fixes a problem where NetIndexFirstBitSegment may not grow to a higher number if NetworkGameplayTagNodeIndex.Num() grows between calls to ConstructNetIndex.

diff --git a/Engine/Source/Runtime/GameplayTags/Private/GameplayTagsManager.cpp b/Engine/Source/Runtime/GameplayTags/Private/GameplayTagsManager.cpp
--- a/Engine/Source/Runtime/GameplayTags/Private/GameplayTagsManager.cpp
+++ b/Engine/Source/Runtime/GameplayTags/Private/GameplayTagsManager.cpp
@@ -582,7 +582,7 @@ void UGameplayTagsManager::ConstructNetIndex()
 	NetIndexTrueBitNum = FMath::CeilToInt(FMath::Log2(static_cast<float>(InvalidTagNetIndex)));
 	
 	// This should never be smaller than NetIndexTrueBitNum
-	NetIndexFirstBitSegment = FMath::Min<int64>(NetIndexFirstBitSegment, NetIndexTrueBitNum);
+	NetIndexFirstBitSegment = FMath::Min<int64>(GetDefault<UGameplayTagsSettings>()->NetIndexFirstBitSegment, NetIndexTrueBitNum);
 
 	// This is now sorted and it should be the same on both client and server
 	if (NetworkGameplayTagNodeIndex.Num() >= INVALID_TAGNETINDEX)
