Iris
- Fixed FGameplayTagContainerNetSerializer::Dequantize

diff --git a/Engine/Source/Runtime/GameplayTags/Private/GameplayTagContainerNetSerializer.cpp b/Engine/Source/Runtime/GameplayTags/Private/GameplayTagContainerNetSerializer.cpp
--- a/Engine/Source/Runtime/GameplayTags/Private/GameplayTagContainerNetSerializer.cpp
+++ b/Engine/Source/Runtime/GameplayTags/Private/GameplayTagContainerNetSerializer.cpp
@@ -141,8 +141,8 @@ void FGameplayTagContainerNetSerializer::Dequantize(FNetSerializationContext& Co
 
 	StructNetSerializer->Dequantize(Context, InternalArgs);
 
-	SourceType& TargetValue = *reinterpret_cast<SourceType*>(Args.Target);
-	TargetValue.CreateFromArray(IntermediateValue.GameplayTags);
+	FGameplayTagContainer& TargetValue = *reinterpret_cast<FGameplayTagContainer*>(Args.Target);
+	TargetValue = FGameplayTagContainer::CreateFromArray(IntermediateValue.GameplayTags);
 }
 
 bool FGameplayTagContainerNetSerializer::IsEqual(FNetSerializationContext& Context, const FNetIsEqualArgs& Args)
