Add support for removing gameplay tags when a gameplay feature is unregistered
With default settings this will only happen in the editor as removing tags at runtime is much riskier
Clean up the tag tree change functions in general

diff --git a/Engine/Plugins/Experimental/GameFeatures/Source/GameFeatures/Private/GameFeaturePluginStateMachine.cpp b/Engine/Plugins/Experimental/GameFeatures/Source/GameFeatures/Private/GameFeaturePluginStateMachine.cpp
--- a/Engine/Plugins/Experimental/GameFeatures/Source/GameFeatures/Private/GameFeaturePluginStateMachine.cpp
+++ b/Engine/Plugins/Experimental/GameFeatures/Source/GameFeatures/Private/GameFeaturePluginStateMachine.cpp
@@ -1796,6 +1796,10 @@ struct FGameFeaturePluginState_Unregistering : public FGameFeaturePluginState
 		verify(FPluginUtils::UnloadPluginAssets(StateProperties.PluginName));
 #endif //if WITH_EDITOR
 
+		// Try to remove the gameplay tags, this might be ignored depending on project settings
+		const FString PluginFolder = FPaths::GetPath(StateProperties.PluginInstalledFilename);
+		UGameplayTagsManager::Get().RemoveTagIniSearchPath(PluginFolder / TEXT("Config") / TEXT("Tags"));
+
 		bRequestedGC = true;
 		GarbageCollectAndUpdateStateMachineDeferred();
 	}
diff --git a/Engine/Source/Runtime/GameplayTags/Classes/GameplayTagsManager.h b/Engine/Source/Runtime/GameplayTags/Classes/GameplayTagsManager.h
--- a/Engine/Source/Runtime/GameplayTags/Classes/GameplayTagsManager.h
+++ b/Engine/Source/Runtime/GameplayTags/Classes/GameplayTagsManager.h
@@ -361,6 +361,7 @@ class GAMEPLAYTAGS_API UGameplayTagsManager : public UObject
 	FGameplayTag AddNativeGameplayTag(FName TagName, const FString& TagDevComment = TEXT("(Native)"));
 
 private:
+	// Only callable from FNativeGameplayTag, these functions do less error checking and can happen after initial tag loading is done
 	void AddNativeGameplayTag(FNativeGameplayTag* TagSource);
 	void RemoveNativeGameplayTag(const FNativeGameplayTag* TagSource);
 
@@ -480,6 +481,9 @@ public:
 	/** Loads tag inis contained in the specified path */
 	void AddTagIniSearchPath(const FString& RootDir);
 
+	/** Tries to remove the specified search path, will return true if anything was removed */
+	bool RemoveTagIniSearchPath(const FString& RootDir);
+
 	/** Gets all the current directories to look for tag sources in */
 	void GetTagSourceSearchPaths(TArray<FString>& OutPaths);
 
@@ -544,6 +548,9 @@ public:
 		return bUseFastReplication;
 	}
 
+	/** If we are allowed to unload tags */
+	bool ShouldUnloadTags() const;
+
 	/** Returns the hash of NetworkGameplayTagNodeIndex */
 	uint32 GetNetworkGameplayTagNodeIndexHash() const { VerifyNetworkIndex(); return NetworkGameplayTagNodeIndexHash; }
 
@@ -782,8 +789,12 @@ private:
 
 	void InvalidateNetworkIndex() { bNetworkIndexInvalidated = true; }
 
+	/** Called in both editor and game when the tag tree changes during startup or editing */
 	void BroadcastOnGameplayTagTreeChanged();
 
+	/** Call after modifying the tag tree nodes, this will either call the full editor refresh or a limited game refresh */
+	void HandleGameplayTagTreeChanged(bool bRecreateTree);
+
 	// Tag Sources
 	///////////////////////////////////////////////////////
 
@@ -821,6 +832,9 @@ private:
 	/** Cached runtime value for whether we should warn when loading invalid tags */
 	bool bShouldClearInvalidTags;
 
+	/** Cached runtime value for whether we should allow unloading of tags */
+	bool bShouldAllowUnloadingTags;
+
 	/** True if native tags have all been added and flushed */
 	bool bDoneAddingNativeTags;
 
diff --git a/Engine/Source/Runtime/GameplayTags/Classes/GameplayTagsSettings.h b/Engine/Source/Runtime/GameplayTags/Classes/GameplayTagsSettings.h
--- a/Engine/Source/Runtime/GameplayTags/Classes/GameplayTagsSettings.h
+++ b/Engine/Source/Runtime/GameplayTags/Classes/GameplayTagsSettings.h
@@ -112,6 +112,14 @@ class GAMEPLAYTAGS_API UGameplayTagsSettings : public UGameplayTagsList
 	UPROPERTY(config, EditAnywhere, Category = GameplayTags, meta = (ConfigRestartRequired = true))
 	bool ClearInvalidTags;
 
+	/** If true, will allow unloading of tags in the editor when plugins are removed */
+	UPROPERTY(config, EditAnywhere, Category = "Advanced Gameplay Tags")
+	bool AllowEditorTagUnloading;
+
+	/** If true, will allow unloading of tags in a non-editor gebuild when plugins are removed, this is potentially unsafe and affects requests to unload during play in editor */
+	UPROPERTY(config, EditAnywhere, Category = "Advanced Gameplay Tags")
+	bool AllowGameTagUnloading;
+
 	/** If true, will replicate gameplay tags by index instead of name. For this to work, tags must be identical on client and server */
 	UPROPERTY(config, EditAnywhere, Category = "Advanced Replication")
 	bool FastReplication;
diff --git a/Engine/Source/Runtime/GameplayTags/Private/GameplayTagsManager.cpp b/Engine/Source/Runtime/GameplayTags/Private/GameplayTagsManager.cpp
--- a/Engine/Source/Runtime/GameplayTags/Private/GameplayTagsManager.cpp
+++ b/Engine/Source/Runtime/GameplayTags/Private/GameplayTagsManager.cpp
@@ -44,6 +44,7 @@ const FName UGameplayTagsManager::NAME_GameplayTagFilter("GameplayTagFilter");
 
 DECLARE_CYCLE_STAT(TEXT("Load Gameplay Tags"), STAT_GameplayTags_LoadGameplayTags, STATGROUP_GameplayTags);
 DECLARE_FLOAT_ACCUMULATOR_STAT(TEXT("Add Tag *.ini Search Path"), STAT_GameplayTags_AddTagIniSearchPath, STATGROUP_GameplayTags);
+DECLARE_FLOAT_ACCUMULATOR_STAT(TEXT("Remove Tag *.ini Search Path"), STAT_GameplayTags_RemoveTagIniSearchPath, STATGROUP_GameplayTags);
 
 #if !(UE_BUILD_SHIPPING)
 
@@ -177,6 +178,7 @@ UGameplayTagsManager::UGameplayTagsManager(const FObjectInitializer& ObjectIniti
 	bShouldWarnOnInvalidTags = true;
 	bShouldClearInvalidTags = false;
 	bDoneAddingNativeTags = false;
+	bShouldAllowUnloadingTags = false;
 	NetIndexFirstBitSegment = 16;
 	NetIndexTrueBitNum = 16;
 	NumBitsForContainerSize = 6;
@@ -284,14 +286,34 @@ void UGameplayTagsManager::AddTagIniSearchPath(const FString& RootDir)
 
 		PathInfo->bWasAddedToTree = true;
 
-		if (!bIsConstructingGameplayTagTree)
-		{
-			InvalidateNetworkIndex();
-			BroadcastOnGameplayTagTreeChanged();
-		}
+		HandleGameplayTagTreeChanged(false);
 	}
 }
 
+bool UGameplayTagsManager::RemoveTagIniSearchPath(const FString& RootDir)
+{
+	SCOPE_SECONDS_ACCUMULATOR(STAT_GameplayTags_RemoveTagIniSearchPath);
+
+	if (!ShouldUnloadTags())
+	{
+		// Can't unload at all
+		return false;
+	}
+
+	FGameplayTagSearchPathInfo* PathInfo = RegisteredSearchPaths.Find(RootDir);
+
+	if (PathInfo)
+	{
+		// Clear out the path and then recreate the tree
+		RegisteredSearchPaths.Remove(RootDir);
+
+		HandleGameplayTagTreeChanged(true);
+
+		return true;
+	}
+	return false;
+}
+
 void UGameplayTagsManager::GetTagSourceSearchPaths(TArray<FString>& OutPaths)
 {
 	OutPaths.Reset();
@@ -528,6 +550,18 @@ void UGameplayTagsManager::ConstructGameplayTagTree()
 			bShouldClearInvalidTags = MutableDefault->ClearInvalidTags;
 			NumBitsForContainerSize = MutableDefault->NumBitsForContainerSize;
 			NetIndexFirstBitSegment = MutableDefault->NetIndexFirstBitSegment;
+
+#if WITH_EDITOR
+			if (GIsEditor)
+			{
+				bShouldAllowUnloadingTags = MutableDefault->AllowEditorTagUnloading;
+			}
+			else
+#endif
+			{
+				bShouldAllowUnloadingTags = MutableDefault->AllowGameTagUnloading;
+			}
+
 		}
 
 		if (ShouldUseFastReplication())
@@ -675,6 +709,26 @@ bool UGameplayTagsManager::ShouldImportTagsFromINI() const
 	return MutableDefault->ImportTagsFromConfig;
 }
 
+bool UGameplayTagsManager::ShouldUnloadTags() const
+{
+#if WITH_EDITOR
+	if (bShouldAllowUnloadingTags && GIsEditor && GEngine)
+	{
+		// Check if we have an active PIE index without linking GEditor, and compare to game setting
+		FWorldContext* PIEWorldContext = GEngine->GetWorldContextFromPIEInstance(0);
+		UGameplayTagsSettings* MutableDefault = GetMutableDefault<UGameplayTagsSettings>();
+
+		if (PIEWorldContext && !MutableDefault->AllowGameTagUnloading)
+		{
+			UE_LOG(LogGameplayTags, Warning, TEXT("Ignoring request to unload tags during Play In Editor because AllowGameTagUnloading=false"));
+			return false;
+		}
+	}
+#endif
+
+	return bShouldAllowUnloadingTags;
+}
+
 void UGameplayTagsManager::GetRestrictedTagConfigFiles(TArray<FString>& RestrictedConfigFiles) const
 {
 	UGameplayTagsSettings* MutableDefault = GetMutableDefault<UGameplayTagsSettings>();
@@ -1088,8 +1142,7 @@ void UGameplayTagsManager::MarkChildrenOfNodeConflict(TSharedPtr<FGameplayTagNod
 #endif
 }
 
-void
-UGameplayTagsManager::BroadcastOnGameplayTagTreeChanged()
+void UGameplayTagsManager::BroadcastOnGameplayTagTreeChanged()
 {
 	if (bDeferBroadcastOnGameplayTagTreeChanged)
 	{
@@ -1101,6 +1154,37 @@ UGameplayTagsManager::BroadcastOnGameplayTagTreeChanged()
 	}
 }
 
+void UGameplayTagsManager::HandleGameplayTagTreeChanged(bool bRecreateTree)
+{
+	// Don't do anything during a reconstruct or before initial native tags are done loading
+	if (!bIsConstructingGameplayTagTree && bDoneAddingNativeTags)
+	{
+		if (bRecreateTree)
+		{
+#if WITH_EDITOR
+			if (GIsEditor)
+			{
+				// In the editor refresh everything
+				EditorRefreshGameplayTagTree();
+				return;
+			}
+#endif
+			DestroyGameplayTagTree();
+			ConstructGameplayTagTree();
+		}
+		else
+		{
+			// Refresh if we're done adding tags
+			if (ShouldUseFastReplication())
+			{
+				InvalidateNetworkIndex();
+			}
+
+			BroadcastOnGameplayTagTreeChanged();
+		}
+	}
+}
+
 UGameplayTagsManager::~UGameplayTagsManager()
 {
 	DestroyGameplayTagTree();
@@ -1943,24 +2027,22 @@ FGameplayTag UGameplayTagsManager::AddNativeGameplayTag(FName TagName, const FSt
 
 void UGameplayTagsManager::AddNativeGameplayTag(FNativeGameplayTag* TagSource)
 {
+	// This adds it to the temporary tree, but expects the caller to add it to FNativeGameplayTag::GetRegisteredNativeTags for later refreshes
 	AddTagTableRow(TagSource->GetGameplayTagTableRow(), FGameplayTagSource::GetNativeName());
-	InvalidateNetworkIndex();
-	BroadcastOnGameplayTagTreeChanged();
+
+	HandleGameplayTagTreeChanged(false);
 }
 
 void UGameplayTagsManager::RemoveNativeGameplayTag(const FNativeGameplayTag* TagSource)
 {
-	// TODO This is awful, need to invalidate the tag tree, not rebuild it.
+	if (!ShouldUnloadTags())
 	{
-#if WITH_EDITOR
-		EditorRefreshGameplayTagTree();
-#else
-	//TODO - Removing tags isn't really a thing right now, but setting this up for the future.
-	//DestroyGameplayTagTree();
-	//ConstructGameplayTagTree();
-	//OnEditorRefreshGameplayTagTree.Broadcast();
-#endif
+		// Ignore if not allowed right now
+		return;
 	}
+
+	// ~FNativeGameplayTag already removed the tag from the global list, so recreate the tree
+	HandleGameplayTagTreeChanged(true);
 }
 
 void UGameplayTagsManager::CallOrRegister_OnDoneAddingNativeTagsDelegate(FSimpleMulticastDelegate::FDelegate Delegate)
diff --git a/Engine/Source/Runtime/GameplayTags/Private/GameplayTagsSettings.cpp b/Engine/Source/Runtime/GameplayTags/Private/GameplayTagsSettings.cpp
--- a/Engine/Source/Runtime/GameplayTags/Private/GameplayTagsSettings.cpp
+++ b/Engine/Source/Runtime/GameplayTags/Private/GameplayTagsSettings.cpp
@@ -63,6 +63,8 @@ UGameplayTagsSettings::UGameplayTagsSettings(const FObjectInitializer& ObjectIni
 	WarnOnInvalidTags = true;
 	ClearInvalidTags = false;
 	FastReplication = false;
+	AllowEditorTagUnloading = true;
+	AllowGameTagUnloading = false;
 	InvalidTagCharacters = ("\"',");
 	NumBitsForContainerSize = 6;
 	NetIndexFirstBitSegment = 16;
