Fix typo in debug log

diff --git a/Engine/Source/Runtime/GameplayTags/Private/GameplayTagsManager.cpp b/Engine/Source/Runtime/GameplayTags/Private/GameplayTagsManager.cpp
--- a/Engine/Source/Runtime/GameplayTags/Private/GameplayTagsManager.cpp
+++ b/Engine/Source/Runtime/GameplayTags/Private/GameplayTagsManager.cpp
@@ -1311,7 +1311,7 @@ void UGameplayTagsManager::PrintReplicationIndices()
 {
 	VerifyNetworkIndex();
 
-	UE_LOG(LogGameplayTags, Display, TEXT("::PrintReplicationIndices (TOTAL %d"), GameplayTagNodeMap.Num());
+	UE_LOG(LogGameplayTags, Display, TEXT("::PrintReplicationIndices (TOTAL %d)"), GameplayTagNodeMap.Num());
 
 	for (auto It : GameplayTagNodeMap)
 	{
