[GameplayTags]
* Changed so validation that define only exists in cpp files is not using stl (to reduce compile times)

diff --git a/Engine/Source/Runtime/GameplayTags/Public/NativeGameplayTags.h b/Engine/Source/Runtime/GameplayTags/Public/NativeGameplayTags.h
--- a/Engine/Source/Runtime/GameplayTags/Public/NativeGameplayTags.h
+++ b/Engine/Source/Runtime/GameplayTags/Public/NativeGameplayTags.h
@@ -11,20 +11,17 @@
 #include "Templates/UnrealTemplate.h"
 #include "UObject/NameTypes.h"
 
-#include <string>
-
 enum class ENativeGameplayTagToken { PRIVATE_USE_MACRO_INSTEAD };
 
 namespace UE::GameplayTags::Private
 {
 	// Used to prevent people from putting UE_DEFINE_GAMEPLAY_TAG_STATIC and UE_DEFINE_GAMEPLAY_TAG in their headers.
-	constexpr bool HasFileExtension(std::string_view file, std::string_view file_ext)
+	constexpr bool HasFileExtension(const char* File)
 	{
-		const auto _Rightsize = file_ext.length();
-		if (file.length() < _Rightsize) {
-			return false;
-		}
-		return file.compare((file.length() - _Rightsize), _Rightsize, file_ext) == 0;
+		const char* It = File;
+		while (*It)
+			++It;
+		return It[-1] == 'p' && It[-2] == 'p' && It[-3] == 'c' && It[-4] == '.';
 	}
 }
 
@@ -36,12 +33,12 @@ namespace UE::GameplayTags::Private
 /**
  * Defines a native gameplay tag that is externally declared in a header to allow other modules or code to use the created tag variable.
  */
-#define UE_DEFINE_GAMEPLAY_TAG(TagName, Tag) FNativeGameplayTag TagName(UE_PLUGIN_NAME, UE_MODULE_NAME, Tag, TEXT(""), ENativeGameplayTagToken::PRIVATE_USE_MACRO_INSTEAD); static_assert(UE::GameplayTags::Private::HasFileExtension(__FILE__, ".cpp"), "UE_DEFINE_GAMEPLAY_TAG can only be used in .cpp files, if you're trying to share tags across modules, use UE_DECLARE_GAMEPLAY_TAG_EXTERN in the public header, and UE_DEFINE_GAMEPLAY_TAG in the private .cpp");
+#define UE_DEFINE_GAMEPLAY_TAG(TagName, Tag) FNativeGameplayTag TagName(UE_PLUGIN_NAME, UE_MODULE_NAME, Tag, TEXT(""), ENativeGameplayTagToken::PRIVATE_USE_MACRO_INSTEAD); static_assert(UE::GameplayTags::Private::HasFileExtension(__FILE__), "UE_DEFINE_GAMEPLAY_TAG can only be used in .cpp files, if you're trying to share tags across modules, use UE_DECLARE_GAMEPLAY_TAG_EXTERN in the public header, and UE_DEFINE_GAMEPLAY_TAG in the private .cpp");
 
 /**
  * Defines a native gameplay tag such that it's only available to the cpp file you define it in.
  */
-#define UE_DEFINE_GAMEPLAY_TAG_STATIC(TagName, Tag) static FNativeGameplayTag TagName(UE_PLUGIN_NAME, UE_MODULE_NAME, Tag, TEXT(""), ENativeGameplayTagToken::PRIVATE_USE_MACRO_INSTEAD); static_assert(UE::GameplayTags::Private::HasFileExtension(__FILE__, ".cpp"), "UE_DEFINE_GAMEPLAY_TAG_STATIC can only be used in .cpp files, if you're trying to share tags across modules, use UE_DECLARE_GAMEPLAY_TAG_EXTERN in the public header, and UE_DEFINE_GAMEPLAY_TAG in the private .cpp");
+#define UE_DEFINE_GAMEPLAY_TAG_STATIC(TagName, Tag) static FNativeGameplayTag TagName(UE_PLUGIN_NAME, UE_MODULE_NAME, Tag, TEXT(""), ENativeGameplayTagToken::PRIVATE_USE_MACRO_INSTEAD); static_assert(UE::GameplayTags::Private::HasFileExtension(__FILE__), "UE_DEFINE_GAMEPLAY_TAG_STATIC can only be used in .cpp files, if you're trying to share tags across modules, use UE_DECLARE_GAMEPLAY_TAG_EXTERN in the public header, and UE_DEFINE_GAMEPLAY_TAG in the private .cpp");
 
 /**
  * Holds a gameplay tag that was registered during static construction of the module, and will
