Add custom serialization support for GameplayTag Queries

diff --git a/Engine/Source/Runtime/GameplayTags/Classes/GameplayTagContainer.h b/Engine/Source/Runtime/GameplayTags/Classes/GameplayTagContainer.h
--- a/Engine/Source/Runtime/GameplayTags/Classes/GameplayTagContainer.h
+++ b/Engine/Source/Runtime/GameplayTags/Classes/GameplayTagContainer.h
@@ -874,6 +874,9 @@ public:
 	/** Builds a FGameplayTagQueryExpression from this query. */
 	void GetQueryExpr(struct FGameplayTagQueryExpression& OutExpr) const;
 
+	/** Serialize the tag query */
+	void Serialize(FArchive& Ar);
+	
 	/** Returns description string. */
 	const FString& GetDescription() const { return UserDescription.IsEmpty() ? AutoDescription : UserDescription; };
 
diff --git a/Engine/Source/Runtime/GameplayTags/Private/GameplayTagContainer.cpp b/Engine/Source/Runtime/GameplayTags/Private/GameplayTagContainer.cpp
--- a/Engine/Source/Runtime/GameplayTags/Private/GameplayTagContainer.cpp
+++ b/Engine/Source/Runtime/GameplayTags/Private/GameplayTagContainer.cpp
@@ -1395,6 +1395,22 @@ void FGameplayTagQuery::GetQueryExpr(FGameplayTagQueryExpression& OutExpr) const
 	QE.Read(OutExpr);
 }
 
+void FGameplayTagQuery::Serialize(FArchive& Ar)
+{
+	Ar << TokenStreamVersion;
+	Ar << TagDictionary;
+	Ar << QueryTokenStream;
+
+#if WITH_EDITOR
+	if (!Ar.IsCooking())
+	{
+		// Descriptions are not needed during runtime
+		Ar << UserDescription;
+		Ar << AutoDescription;
+	}
+#endif
+}
+
 void FGameplayTagQuery::Build(FGameplayTagQueryExpression& RootQueryExpr, FString InUserDescription)
 {
 	TokenStreamVersion = EGameplayTagQueryStreamVersion::LatestVersion;
