#UE Add accessor to get the number of elements in the gameplay tags search path list, so you don't need to build an array to get the number of paths

diff --git a/Engine/Source/Runtime/GameplayTags/Classes/GameplayTagsManager.h b/Engine/Source/Runtime/GameplayTags/Classes/GameplayTagsManager.h
--- a/Engine/Source/Runtime/GameplayTags/Classes/GameplayTagsManager.h
+++ b/Engine/Source/Runtime/GameplayTags/Classes/GameplayTagsManager.h
@@ -487,6 +487,9 @@ public:
 	/** Gets all the current directories to look for tag sources in */
 	void GetTagSourceSearchPaths(TArray<FString>& OutPaths);
 
+	/** Gets the number of tag source search paths */
+	int32 GetNumTagSourceSearchPaths();
+
 	/** Helper function to construct the gameplay tag tree */
 	void ConstructGameplayTagTree();
 
diff --git a/Engine/Source/Runtime/GameplayTags/Private/GameplayTagsManager.cpp b/Engine/Source/Runtime/GameplayTags/Private/GameplayTagsManager.cpp
--- a/Engine/Source/Runtime/GameplayTags/Private/GameplayTagsManager.cpp
+++ b/Engine/Source/Runtime/GameplayTags/Private/GameplayTagsManager.cpp
@@ -322,6 +322,11 @@ void UGameplayTagsManager::GetTagSourceSearchPaths(TArray<FString>& OutPaths)
 	RegisteredSearchPaths.GenerateKeyArray(OutPaths);
 }
 
+int32 UGameplayTagsManager::GetNumTagSourceSearchPaths()
+{
+	return RegisteredSearchPaths.Num();
+}
+
 void UGameplayTagsManager::AddRestrictedGameplayTagSource(const FString& FileName)
 {
 	FName TagSource = FName(*FPaths::GetCleanFilename(FileName));
