Don't try to read an empty GameplayTagQuery when creating an editable query

diff --git a/Engine/Source/Runtime/GameplayTags/Private/GameplayTagContainer.cpp b/Engine/Source/Runtime/GameplayTags/Private/GameplayTagContainer.cpp
--- a/Engine/Source/Runtime/GameplayTags/Private/GameplayTagContainer.cpp
+++ b/Engine/Source/Runtime/GameplayTags/Private/GameplayTagContainer.cpp
@@ -169,7 +169,7 @@ private:
 		}
 		
 		UE_LOG(LogGameplayTags, Warning, TEXT("Error parsing FGameplayTagQuery!"));
-		bReadError = true;
+		bReadError = true; 
 		return 0;
 	}
 };
@@ -1490,16 +1490,19 @@ UEditableGameplayTagQuery* FQueryEvaluator::CreateEditableQuery()
 	UEditableGameplayTagQuery* const EditableQuery = NewObject<UEditableGameplayTagQuery>(GetTransientPackage(), NAME_None, RF_Transactional);
 
 	// start parsing the set
-	Version = GetToken();
-	if (!bReadError)
+	if (!Query.IsEmpty())
 	{
-		uint8 const bHasRootExpression = GetToken();
-		if (!bReadError && bHasRootExpression)
+		Version = GetToken();
+		if (!bReadError)
 		{
-			EditableQuery->RootExpression = ReadEditableQueryExpr(EditableQuery);
+			uint8 const bHasRootExpression = GetToken();
+			if (!bReadError && bHasRootExpression)
+			{
+				EditableQuery->RootExpression = ReadEditableQueryExpr(EditableQuery);
+			}
 		}
+		ensure(CurStreamIdx == Query.QueryTokenStream.Num());
 	}
-	ensure(CurStreamIdx == Query.QueryTokenStream.Num());
 
 	EditableQuery->UserDescription = Query.UserDescription;
 
