PR #9758: Fixed FGameplayTag::ImportTextItem so that Buffer is updated when we read tokens from it - addresses an issue where copying and pasting a gameplay tag array in editor would only import the first element (Contributed by [REDACTED_PERSON])

diff --git a/Engine/Source/Runtime/GameplayTags/Private/GameplayTagContainer.cpp b/Engine/Source/Runtime/GameplayTags/Private/GameplayTagContainer.cpp
--- a/Engine/Source/Runtime/GameplayTags/Private/GameplayTagContainer.cpp
+++ b/Engine/Source/Runtime/GameplayTags/Private/GameplayTagContainer.cpp
@@ -1293,13 +1293,16 @@ void FGameplayTag::PostSerialize(const FArchive& Ar)
 
 bool FGameplayTag::ImportTextItem(const TCHAR*& Buffer, int32 PortFlags, UObject* Parent, FOutputDevice* ErrorText)
 {
-	FString ImportedTag = TEXT("");
+	FString ImportedTag;
 	const TCHAR* NewBuffer = FPropertyHelpers::ReadToken(Buffer, ImportedTag, true);
 	if (!NewBuffer)
 	{
 		// Failed to read buffer. Maybe normal ImportText will work.
 		return false;
 	}
+	
+	const TCHAR* OriginalBuffer = Buffer;
+	Buffer = NewBuffer;
 
 	if (ImportedTag == TEXT("None") || ImportedTag.IsEmpty())
 	{
@@ -1312,7 +1315,7 @@ bool FGameplayTag::ImportTextItem(const TCHAR*& Buffer, int32 PortFlags, UObject
 	{
 		// Let normal ImportText handle this before handling fixups
 		UScriptStruct* ScriptStruct = FGameplayTag::StaticStruct();
-		Buffer = ScriptStruct->ImportText(Buffer, this, Parent, PortFlags, ErrorText, ScriptStruct->GetName(), false);
+		Buffer = ScriptStruct->ImportText(OriginalBuffer, this, Parent, PortFlags, ErrorText, ScriptStruct->GetName(), false);
 		UGameplayTagsManager::Get().ImportSingleGameplayTag(*this, TagName, !!(PortFlags & PPF_SerializedAsImportText));
 		return true;
 	}
