GameplayTagsManager hooks to suspend/resume EditorRefreshGameplayTagTree

diff --git a/Engine/Source/Runtime/GameplayTags/Classes/GameplayTagsManager.h b/Engine/Source/Runtime/GameplayTags/Classes/GameplayTagsManager.h
--- a/Engine/Source/Runtime/GameplayTags/Classes/GameplayTagsManager.h
+++ b/Engine/Source/Runtime/GameplayTags/Classes/GameplayTagsManager.h
@@ -655,9 +655,20 @@ public:
 	/** Returns information about tag. If not found return false */
 	bool GetTagEditorData(FName TagName, FString& OutComment, FName &OutTagSource, bool& bOutIsTagExplicit, bool &bOutIsRestrictedTag, bool &bOutAllowNonRestrictedChildren) const;
 
+#if WITH_EDITOR
+	/** This is called after EditorRefreshGameplayTagTree. Useful if you need to do anything editor related when tags are added or removed */
+	static FSimpleMulticastDelegate OnEditorRefreshGameplayTagTree;
+
 	/** Refresh the gameplaytag tree due to an editor change */
 	void EditorRefreshGameplayTagTree();
 
+	/** Suspends EditorRefreshGameplayTagTree requests */
+	void SuspendEditorRefreshGameplayTagTree(FGuid SuspendToken);
+
+	/** Resumes EditorRefreshGameplayTagTree requests; triggers a refresh if a request was made while it was suspended */
+	void ResumeEditorRefreshGameplayTagTree(FGuid SuspendToken);
+#endif //if WITH_EDITOR
+
 	/** Gets a Tag Container containing all of the tags in the hierarchy that are children of this tag, and were explicitly added to the dictionary */
 	FGameplayTagContainer RequestGameplayTagChildrenInDictionary(const FGameplayTag& GameplayTag) const;
 #if WITH_EDITORONLY_DATA
@@ -665,8 +676,6 @@ public:
 	FGameplayTagContainer RequestGameplayTagDirectDescendantsInDictionary(const FGameplayTag& GameplayTag, EGameplayTagSelectionType SelectionType) const;
 #endif // WITH_EDITORONLY_DATA
 
-	/** This is called when EditorRefreshGameplayTagTree. Useful if you need to do anything editor related when tags are added or removed */
-	static FSimpleMulticastDelegate OnEditorRefreshGameplayTagTree;
 
 	DECLARE_MULTICAST_DELEGATE_TwoParams(FOnGameplayTagDoubleClickedEditor, FGameplayTag, FSimpleMulticastDelegate& /* OUT */)
 	FOnGameplayTagDoubleClickedEditor OnGatherGameplayTagDoubleClickedEditor;
@@ -854,7 +863,10 @@ private:
 
 	// Transient editor-only tags to support quick-iteration PIE workflows
 	TSet<FName> TransientEditorTags;
-#endif
+
+	TSet<FGuid> EditorRefreshGameplayTagTreeSuspendTokens;
+	bool bEditorRefreshGameplayTagTreeRequestedDuringSuspend = false;
+#endif //if WITH_EDITOR
 
 	/** Sorted list of nodes, used for network replication */
 	TArray<TSharedPtr<FGameplayTagNode>> NetworkGameplayTagNodeIndex;
diff --git a/Engine/Source/Runtime/GameplayTags/Private/GameplayTagsManager.cpp b/Engine/Source/Runtime/GameplayTags/Private/GameplayTagsManager.cpp
--- a/Engine/Source/Runtime/GameplayTags/Private/GameplayTagsManager.cpp
+++ b/Engine/Source/Runtime/GameplayTags/Private/GameplayTagsManager.cpp
@@ -1679,8 +1679,16 @@ bool UGameplayTagsManager::GetTagEditorData(FName TagName, FString& OutComment,
 	return false;
 }
 
+#if WITH_EDITOR
+
 void UGameplayTagsManager::EditorRefreshGameplayTagTree()
 {
+	if (!EditorRefreshGameplayTagTreeSuspendTokens.IsEmpty())
+	{
+		bEditorRefreshGameplayTagTreeRequestedDuringSuspend = true;
+		return;
+	}
+
 	TRACE_CPUPROFILER_EVENT_SCOPE(UGameplayTagsManager::EditorRefreshGameplayTagTree)
 
 	// Clear out source path info so it will reload off disk
@@ -1696,6 +1704,23 @@ void UGameplayTagsManager::EditorRefreshGameplayTagTree()
 	OnEditorRefreshGameplayTagTree.Broadcast();
 }
 
+void UGameplayTagsManager::SuspendEditorRefreshGameplayTagTree(FGuid SuspendToken)
+{
+	EditorRefreshGameplayTagTreeSuspendTokens.Add(SuspendToken);
+}
+
+void UGameplayTagsManager::ResumeEditorRefreshGameplayTagTree(FGuid SuspendToken)
+{
+	EditorRefreshGameplayTagTreeSuspendTokens.Remove(SuspendToken);
+	if (EditorRefreshGameplayTagTreeSuspendTokens.IsEmpty() && bEditorRefreshGameplayTagTreeRequestedDuringSuspend)
+	{
+		bEditorRefreshGameplayTagTreeRequestedDuringSuspend = false;
+		EditorRefreshGameplayTagTree();
+;	}
+}
+
+#endif //if WITH_EDITOR
+
 FGameplayTagContainer UGameplayTagsManager::RequestGameplayTagChildrenInDictionary(const FGameplayTag& GameplayTag) const
 {
 	// Note this purposefully does not include the passed in GameplayTag in the container.
