Adding hashing for the GameplayTagContainer

#rnx

diff --git a/Engine/Source/Runtime/GameplayTags/Classes/GameplayTagContainer.h b/Engine/Source/Runtime/GameplayTags/Classes/GameplayTagContainer.h
--- a/Engine/Source/Runtime/GameplayTags/Classes/GameplayTagContainer.h
+++ b/Engine/Source/Runtime/GameplayTags/Classes/GameplayTagContainer.h
@@ -422,6 +422,16 @@ struct GAMEPLAYTAGS_API FGameplayTagContainer
 		return GameplayTags.Num() == 0;
 	}
 
+	FORCEINLINE friend uint32 GetTypeHash(const FGameplayTagContainer& Tags)
+	{
+		uint32 HashValue = 0;
+		for (const FGameplayTag& Tag : Tags)
+		{
+			HashValue = HashCombine(GetTypeHash(Tag), HashValue);
+		}
+		return HashValue;
+	}
+
 	/** Returns a new container explicitly containing the tags of this container and all of their parent tags */
 	FGameplayTagContainer GetGameplayTagParents() const;
 
