Native Gameplay Tags - We now track the multiple sources a single gameplay tag can actually come from so that we can determine if it's actually referenced by a plugin that will be cooked.

diff --git a/Engine/Plugins/Editor/GameplayTagsEditor/Source/GameplayTagsEditor/Private/SGameplayTagWidget.cpp b/Engine/Plugins/Editor/GameplayTagsEditor/Source/GameplayTagsEditor/Private/SGameplayTagWidget.cpp
--- a/Engine/Plugins/Editor/GameplayTagsEditor/Source/GameplayTagsEditor/Private/SGameplayTagWidget.cpp
+++ b/Engine/Plugins/Editor/GameplayTagsEditor/Source/GameplayTagsEditor/Private/SGameplayTagWidget.cpp
@@ -111,7 +111,7 @@ void SGameplayTagWidget::Construct(const FArguments& InArgs, const TArray<FEdita
 		for (int32 Idx = TagItems.Num() - 1; Idx >= 0; --Idx)
 		{
 			bool DelegateShouldHide = false;
-			FGameplayTagSource* Source = Manager.FindTagSource(TagItems[Idx]->SourceName);
+			FGameplayTagSource* Source = Manager.FindTagSource(TagItems[Idx]->GetFirstSourceName());
 			Manager.OnFilterGameplayTag.Broadcast(UGameplayTagsManager::FFilterGameplayTagContext(RootFilterString, TagItems[Idx], Source, PropertyHandle), DelegateShouldHide);
 			if (DelegateShouldHide)
 			{
@@ -425,7 +425,7 @@ bool SGameplayTagWidget::FilterChildrenCheck( TSharedPtr<FGameplayTagNode> InIte
 	Manager.OnFilterGameplayTagChildren.Broadcast(RootFilterString, InItem, DelegateShouldHide);
 	if (!DelegateShouldHide && Manager.OnFilterGameplayTag.IsBound())
 	{
-		FGameplayTagSource* Source = Manager.FindTagSource(InItem->SourceName);
+		FGameplayTagSource* Source = Manager.FindTagSource(InItem->GetFirstSourceName());
 		Manager.OnFilterGameplayTag.Broadcast(UGameplayTagsManager::FFilterGameplayTagContext(RootFilterString, InItem, Source, PropertyHandle), DelegateShouldHide);
 	}
 	if (DelegateShouldHide)
@@ -463,7 +463,7 @@ TSharedRef<ITableRow> SGameplayTagWidget::OnGenerateRow(TSharedPtr<FGameplayTagN
 
 				if (Node->bIsExplicitTag)
 				{
-					TagSource = Node->SourceName;
+					TagSource = Node->GetFirstSourceName();
 				}
 				else
 				{
@@ -930,7 +930,7 @@ TSharedRef<SWidget> SGameplayTagWidget::MakeTagActionsMenu(TSharedPtr<FGameplayT
 	}
 
 	// we can only rename or delete tags if they came from an ini file
-	if (!InTagNode->SourceName.ToString().EndsWith(TEXT(".ini")))
+	if (!InTagNode->GetFirstSourceName().ToString().EndsWith(TEXT(".ini")))
 	{
 		bShowManagement = false;
 	}
@@ -1255,7 +1255,7 @@ void SGameplayTagWidget::RefreshTags()
 		for (int32 Idx = TagItems.Num() - 1; Idx >= 0; --Idx)
 		{
 			bool DelegateShouldHide = false;
-			FGameplayTagSource* Source = Manager.FindTagSource(TagItems[Idx]->SourceName);
+			FGameplayTagSource* Source = Manager.FindTagSource(TagItems[Idx]->GetFirstSourceName());
 			Manager.OnFilterGameplayTag.Broadcast(UGameplayTagsManager::FFilterGameplayTagContext(RootFilterString, TagItems[Idx], Source, PropertyHandle), DelegateShouldHide);
 			if (DelegateShouldHide)
 			{
diff --git a/Engine/Source/Runtime/GameplayTags/Classes/GameplayTagsManager.h b/Engine/Source/Runtime/GameplayTags/Classes/GameplayTagsManager.h
--- a/Engine/Source/Runtime/GameplayTags/Classes/GameplayTagsManager.h
+++ b/Engine/Source/Runtime/GameplayTags/Classes/GameplayTagsManager.h
@@ -245,6 +245,11 @@ struct FGameplayTagNode
 		return true;
 	}
 
+#if WITH_EDITORONLY_DATA
+	FName GetFirstSourceName() const { return SourceNames.Num() == 0 ? NAME_None : SourceNames[0]; }
+	const TArray<FName>& GetAllSourceNames() const { return SourceNames; }
+#endif
+
 private:
 	/** Raw name for this tag at current rank in the tree */
 	FName Tag;
@@ -262,8 +267,8 @@ private:
 	FGameplayTagNetIndex NetIndex;
 
 #if WITH_EDITORONLY_DATA
-	/** Package or config file this tag came from. This is the first one added. If None, this is an implicitly added tag */
-	FName SourceName;
+	/** Module or Package or config file this tag came from. If empty this is an implicitly added tag */
+	TArray<FName> SourceNames;
 
 	/** Comment for this tag */
 	FString DevComment;
@@ -514,6 +519,8 @@ public:
 	/** Fills in an array with all tag sources of a specific type */
 	void FindTagSourcesWithType(EGameplayTagSourceType TagSourceType, TArray<const FGameplayTagSource*>& OutArray) const;
 
+	void FindTagsWithSource(FStringView PackageNameOrPath, TArray<FGameplayTag>& OutTags) const;
+
 	/**
 	 * Check to see how closely two FGameplayTags match. Higher values indicate more matching terms in the tags.
 	 *
@@ -653,7 +660,10 @@ public:
 	bool IsDictionaryTag(FName TagName) const;
 
 	/** Returns information about tag. If not found return false */
-	bool GetTagEditorData(FName TagName, FString& OutComment, FName &OutTagSource, bool& bOutIsTagExplicit, bool &bOutIsRestrictedTag, bool &bOutAllowNonRestrictedChildren) const;
+	bool GetTagEditorData(FName TagName, FString& OutComment, FName &OutFirstTagSource, bool& bOutIsTagExplicit, bool &bOutIsRestrictedTag, bool &bOutAllowNonRestrictedChildren) const;
+	
+	/** Returns information about tag. If not found return false */
+    bool GetTagEditorData(FName TagName, FString& OutComment, TArray<FName>& OutTagSources, bool& bOutIsTagExplicit, bool &bOutIsRestrictedTag, bool &bOutAllowNonRestrictedChildren) const;
 
 #if WITH_EDITOR
 	/** This is called after EditorRefreshGameplayTagTree. Useful if you need to do anything editor related when tags are added or removed */
diff --git a/Engine/Source/Runtime/GameplayTags/Private/GameplayTagsManager.cpp b/Engine/Source/Runtime/GameplayTags/Private/GameplayTagsManager.cpp
--- a/Engine/Source/Runtime/GameplayTags/Private/GameplayTagsManager.cpp
+++ b/Engine/Source/Runtime/GameplayTags/Private/GameplayTagsManager.cpp
@@ -454,7 +454,11 @@ void UGameplayTagsManager::ConstructGameplayTagTree()
 
 			for (const class FNativeGameplayTag* NativeTag : FNativeGameplayTag::GetRegisteredNativeTags())
 			{
+#if WITH_EDITOR && !UE_BUILD_SHIPPING
+				AddTagTableRow(NativeTag->GetGameplayTagTableRow(), NativeTag->GetModuleName());
+#else
 				AddTagTableRow(NativeTag->GetGameplayTagTableRow(), FGameplayTagSource::GetNativeName());
+#endif
 			}
 		}
 
@@ -1104,7 +1108,7 @@ void UGameplayTagsManager::AddTagTableRow(const FGameplayTagTableRow& TagRow, FN
 			CurNode->bAncestorHasConflict = bHasSeenConflict;
 
 			// If the sources don't match and the tag is explicit and we should've added the tag explicitly here, we have a conflict
-			if (CurNode->SourceName != SourceName && (CurNode->bIsExplicitTag && bIsExplicitTag))
+			if (CurNode->GetFirstSourceName() != SourceName && (CurNode->bIsExplicitTag && bIsExplicitTag))
 			{
 				// mark all ancestors as having a bad descendant
 				for (TSharedPtr<FGameplayTagNode> CurAncestorNode : AncestorNodes)
@@ -1116,10 +1120,7 @@ void UGameplayTagsManager::AddTagTableRow(const FGameplayTagTableRow& TagRow, FN
 				CurNode->bNodeHasConflict = true;
 
 				// append source names
-				FString CombinedSources = CurNode->SourceName.ToString();
-				CombinedSources.Append(TEXT(" and "));
-				CombinedSources.Append(SourceName.ToString());
-				CurNode->SourceName = FName(*CombinedSources);
+				CurNode->SourceNames.Add(SourceName);
 
 				// mark all current descendants as having a bad ancestor
 				MarkChildrenOfNodeConflict(CurNode);
@@ -1292,18 +1293,8 @@ int32 UGameplayTagsManager::InsertTagIntoNodeArray(FName Tag, FName FullTag, TSh
 	}
 
 #if WITH_EDITOR
-	static FName NativeSourceName = FGameplayTagSource::GetNativeName();
-
 	// Set/update editor only data
-	if (NodeArray[FoundNodeIdx]->SourceName.IsNone() && !SourceName.IsNone())
-	{
-		NodeArray[FoundNodeIdx]->SourceName = SourceName;
-	}
-	else if (SourceName == NativeSourceName)
-	{
-		// Native overrides other types
-		NodeArray[FoundNodeIdx]->SourceName = SourceName;
-	}
+	NodeArray[FoundNodeIdx]->SourceNames.AddUnique(SourceName);
 
 	if (NodeArray[FoundNodeIdx]->DevComment.IsEmpty() && !DevComment.IsEmpty())
 	{
@@ -1646,7 +1637,7 @@ void UGameplayTagsManager::GetAllTagsFromSource(FName TagSource, TArray< TShared
 {
 	for (const TPair<FGameplayTag, TSharedPtr<FGameplayTagNode>>& NodePair : GameplayTagNodeMap)
 	{
-		if (NodePair.Value->SourceName == TagSource)
+		if (NodePair.Value->SourceNames.Contains(TagSource))
 		{
 			OutTagArray.Add(NodePair.Value);
 		}
@@ -1664,13 +1655,28 @@ bool UGameplayTagsManager::IsDictionaryTag(FName TagName) const
 	return false;
 }
 
-bool UGameplayTagsManager::GetTagEditorData(FName TagName, FString& OutComment, FName& OutTagSource, bool& bOutIsTagExplicit, bool &bOutIsRestrictedTag, bool &bOutAllowNonRestrictedChildren) const
+bool UGameplayTagsManager::GetTagEditorData(FName TagName, FString& OutComment, FName& OutFirstTagSource, bool& bOutIsTagExplicit, bool &bOutIsRestrictedTag, bool &bOutAllowNonRestrictedChildren) const
 {
 	TSharedPtr<FGameplayTagNode> Node = FindTagNode(TagName);
 	if (Node.IsValid())
 	{
 		OutComment = Node->DevComment;
-		OutTagSource = Node->SourceName;
+		OutFirstTagSource = Node->GetFirstSourceName();
+		bOutIsTagExplicit = Node->bIsExplicitTag;
+		bOutIsRestrictedTag = Node->bIsRestrictedTag;
+		bOutAllowNonRestrictedChildren = Node->bAllowNonRestrictedChildren;
+		return true;
+	}
+	return false;
+}
+
+bool UGameplayTagsManager::GetTagEditorData(FName TagName, FString& OutComment, TArray<FName>& OutTagSources, bool& bOutIsTagExplicit, bool &bOutIsRestrictedTag, bool &bOutAllowNonRestrictedChildren) const
+{
+	TSharedPtr<FGameplayTagNode> Node = FindTagNode(TagName);
+	if (Node.IsValid())
+	{
+		OutComment = Node->DevComment;
+		OutTagSources = Node->GetAllSourceNames();
 		bOutIsTagExplicit = Node->bIsExplicitTag;
 		bOutIsRestrictedTag = Node->bIsRestrictedTag;
 		bOutAllowNonRestrictedChildren = Node->bAllowNonRestrictedChildren;
@@ -1754,7 +1760,7 @@ FGameplayTagContainer UGameplayTagsManager::RequestGameplayTagDirectDescendantsI
 			if (ChildNode.IsValid())
 			{
 				// if the tag isn't in the dictionary, add its children to the list
-				if (ChildNode->SourceName == NAME_None)
+				if (ChildNode->GetFirstSourceName() == NAME_None)
 				{
 					TArray< TSharedPtr<FGameplayTagNode> >& GrandChildrenNodes = ChildNode->GetChildTagNodes();
 					ChildrenNodes.Append(GrandChildrenNodes);
@@ -1811,6 +1817,48 @@ FGameplayTagSource* UGameplayTagsManager::FindTagSource(FName TagSourceName)
 	return TagSources.Find(TagSourceName);
 }
 
+void UGameplayTagsManager::FindTagsWithSource(FStringView PackageNameOrPath, TArray<FGameplayTag>& OutTags) const
+{
+	for (const auto& TagSourceEntry : TagSources)
+	{
+		const FGameplayTagSource& Source = TagSourceEntry.Value;
+		
+		FString SourcePackagePath;
+		switch (Source.SourceType)
+		{
+		case EGameplayTagSourceType::TagList:
+			if (Source.SourceTagList)
+			{
+				const FString ContentFilePath = FPaths::GetPath(Source.SourceTagList->ConfigFileName) / TEXT("../../Content/");
+				FString RootContentPath;
+				if (FPackageName::TryConvertFilenameToLongPackageName(ContentFilePath, RootContentPath))
+				{
+					SourcePackagePath = *RootContentPath;
+				}
+			}
+			break;
+		case EGameplayTagSourceType::DataTable:
+			SourcePackagePath = Source.SourceName.ToString();
+			break;
+		case EGameplayTagSourceType::Native:
+			SourcePackagePath = Source.SourceName.ToString();
+		default:
+			break;
+		}
+
+		if (SourcePackagePath.StartsWith(PackageNameOrPath))
+		{
+			if (Source.SourceTagList)
+			{
+				for (const FGameplayTagTableRow& Row : Source.SourceTagList->GameplayTagList)
+				{
+					OutTags.Add(FGameplayTag(Row.Tag));
+				}
+			}
+		}
+	}
+}
+
 void UGameplayTagsManager::FindTagSourcesWithType(EGameplayTagSourceType TagSourceType, TArray<const FGameplayTagSource*>& OutArray) const
 {
 	for (auto TagSourceIt = TagSources.CreateConstIterator(); TagSourceIt; ++TagSourceIt)
@@ -1839,7 +1887,11 @@ FGameplayTagSource* UGameplayTagsManager::FindOrAddTagSource(FName TagSourceName
 
 	FGameplayTagSource* NewSource = &TagSources.Add(TagSourceName, FGameplayTagSource(TagSourceName, SourceType));
 
-	if (SourceType == EGameplayTagSourceType::DefaultTagList)
+	if (SourceType == EGameplayTagSourceType::Native)
+	{
+		NewSource->SourceTagList = NewObject<UGameplayTagsList>(this, TagSourceName, RF_Transient);
+	}
+	else if (SourceType == EGameplayTagSourceType::DefaultTagList)
 	{
 		NewSource->SourceTagList = GetMutableDefault<UGameplayTagsSettings>();
 	}
@@ -2059,8 +2111,11 @@ FGameplayTag UGameplayTagsManager::AddNativeGameplayTag(FName TagName, const FSt
 
 void UGameplayTagsManager::AddNativeGameplayTag(FNativeGameplayTag* TagSource)
 {
+	const FGameplayTagSource* NativeSource = FindOrAddTagSource(TagSource->GetModuleName(), EGameplayTagSourceType::Native);
+	NativeSource->SourceTagList->GameplayTagList.Add(TagSource->GetGameplayTagTableRow());
+	
 	// This adds it to the temporary tree, but expects the caller to add it to FNativeGameplayTag::GetRegisteredNativeTags for later refreshes
-	AddTagTableRow(TagSource->GetGameplayTagTableRow(), FGameplayTagSource::GetNativeName());
+	AddTagTableRow(TagSource->GetGameplayTagTableRow(), NativeSource->SourceName);
 
 	HandleGameplayTagTreeChanged(false);
 }
@@ -2194,7 +2249,7 @@ void UGameplayTagsManager::AddChildrenTags(FGameplayTagContainer& TagContainer,
 				bool bShouldInclude = true;
 
 #if WITH_EDITORONLY_DATA
-				if (OnlyIncludeDictionaryTags && ChildNode->SourceName == NAME_None)
+				if (OnlyIncludeDictionaryTags && ChildNode->GetFirstSourceName() == NAME_None)
 				{
 					// Only have info to do this in editor builds
 					bShouldInclude = false;
@@ -2409,7 +2464,7 @@ void FGameplayTagNode::ResetNode()
 	ParentNode.Reset();
 
 #if WITH_EDITORONLY_DATA
-	SourceName = NAME_None;
+	SourceNames.Reset();
 	DevComment = "";
 	bIsExplicitTag = false;
 	bIsRestrictedTag = false;
diff --git a/Engine/Source/Runtime/GameplayTags/Private/NativeGameplayTags.cpp b/Engine/Source/Runtime/GameplayTags/Private/NativeGameplayTags.cpp
--- a/Engine/Source/Runtime/GameplayTags/Private/NativeGameplayTags.cpp
+++ b/Engine/Source/Runtime/GameplayTags/Private/NativeGameplayTags.cpp
@@ -6,6 +6,7 @@
 #include "ProjectDescriptor.h"
 #include "PluginDescriptor.h"
 #include "Interfaces/IPluginManager.h"
+#include "Misc/PackageName.h"
 
 #define LOCTEXT_NAMESPACE "FNativeGameplayTag"
 
@@ -59,6 +60,7 @@ FNativeGameplayTag::FNativeGameplayTag(FName InPluginName, FName InModuleName, F
 #if !UE_BUILD_SHIPPING
 	PluginName = InPluginName;
 	ModuleName = InModuleName;
+	ModulePackageName = FPackageName::GetModuleScriptPackageName(InModuleName);
 #endif
 
 	InternalTag = TagName.IsNone() ? FGameplayTag() : FGameplayTag(TagName);
diff --git a/Engine/Source/Runtime/GameplayTags/Public/NativeGameplayTags.h b/Engine/Source/Runtime/GameplayTags/Public/NativeGameplayTags.h
--- a/Engine/Source/Runtime/GameplayTags/Public/NativeGameplayTags.h
+++ b/Engine/Source/Runtime/GameplayTags/Public/NativeGameplayTags.h
@@ -74,12 +74,19 @@ public:
 #endif
 	}
 
+#if WITH_EDITOR && !UE_BUILD_SHIPPING
+	FName GetPlugin() const { return PluginName; }
+	FName GetModuleName() const { return ModuleName; }
+	FName GetModulePackageName() const { return ModulePackageName; }
+#endif
+
 private:
 	FGameplayTag InternalTag;
 
 #if !UE_BUILD_SHIPPING
 	FName PluginName;
 	FName ModuleName;
+	FName ModulePackageName;
 	mutable bool bValidated = false;
 
 	void ValidateTagRegistration() const;
