NativeGameplayTags - Adding a fallback on builds without module access to the metadata for tag source.

diff --git a/Engine/Source/Runtime/GameplayTags/Private/GameplayTagsManager.cpp b/Engine/Source/Runtime/GameplayTags/Private/GameplayTagsManager.cpp
--- a/Engine/Source/Runtime/GameplayTags/Private/GameplayTagsManager.cpp
+++ b/Engine/Source/Runtime/GameplayTags/Private/GameplayTagsManager.cpp
@@ -454,7 +454,7 @@ void UGameplayTagsManager::ConstructGameplayTagTree()
 
 			for (const class FNativeGameplayTag* NativeTag : FNativeGameplayTag::GetRegisteredNativeTags())
 			{
-#if WITH_EDITOR && !UE_BUILD_SHIPPING
+#if UE_INCLUDE_NATIVE_GAMEPLAYTAG_METADATA
 				AddTagTableRow(NativeTag->GetGameplayTagTableRow(), NativeTag->GetModuleName());
 #else
 				AddTagTableRow(NativeTag->GetGameplayTagTableRow(), FGameplayTagSource::GetNativeName());
diff --git a/Engine/Source/Runtime/GameplayTags/Public/NativeGameplayTags.h b/Engine/Source/Runtime/GameplayTags/Public/NativeGameplayTags.h
--- a/Engine/Source/Runtime/GameplayTags/Public/NativeGameplayTags.h
+++ b/Engine/Source/Runtime/GameplayTags/Public/NativeGameplayTags.h
@@ -45,6 +45,10 @@ namespace UE::GameplayTags::Private
  */
 #define UE_DEFINE_GAMEPLAY_TAG_STATIC(TagName, Tag) static FNativeGameplayTag TagName(UE_PLUGIN_NAME, UE_MODULE_NAME, Tag, TEXT(""), ENativeGameplayTagToken::PRIVATE_USE_MACRO_INSTEAD); static_assert(UE::GameplayTags::Private::HasFileExtension(__FILE__), "UE_DEFINE_GAMEPLAY_TAG_STATIC can only be used in .cpp files, if you're trying to share tags across modules, use UE_DECLARE_GAMEPLAY_TAG_EXTERN in the public header, and UE_DEFINE_GAMEPLAY_TAG in the private .cpp");
 
+#ifndef UE_INCLUDE_NATIVE_GAMEPLAYTAG_METADATA
+	#define UE_INCLUDE_NATIVE_GAMEPLAYTAG_METADATA WITH_EDITOR && !UE_BUILD_SHIPPING
+#endif
+
 /**
  * Holds a gameplay tag that was registered during static construction of the module, and will
  * be unregistered when the module unloads.  Each registration is based on the native tag pointer
@@ -74,10 +78,12 @@ public:
 #endif
 	}
 
-#if WITH_EDITOR && !UE_BUILD_SHIPPING
+#if UE_INCLUDE_NATIVE_GAMEPLAYTAG_METADATA
 	FName GetPlugin() const { return PluginName; }
 	FName GetModuleName() const { return ModuleName; }
 	FName GetModulePackageName() const { return ModulePackageName; }
+#else
+	FName GetModuleName() const { return NAME_Native; }
 #endif
 
 private:
