NativeGameplayTag - Need to redefine NAME_Native in the gameplay tag.

diff --git a/Engine/Source/Runtime/GameplayTags/Private/NativeGameplayTags.cpp b/Engine/Source/Runtime/GameplayTags/Private/NativeGameplayTags.cpp
--- a/Engine/Source/Runtime/GameplayTags/Private/NativeGameplayTags.cpp
+++ b/Engine/Source/Runtime/GameplayTags/Private/NativeGameplayTags.cpp
@@ -50,6 +50,11 @@ static bool VerifyModuleCanContainGameplayTag(FName ModuleName, FName TagName, c
 
 #endif
 
+namespace UE::NativeGameplayTag
+{
+	const FName NAME_Native = FName(TEXT("Native"));
+}
+
 FNativeGameplayTag::FNativeGameplayTag(FName InPluginName, FName InModuleName, FName TagName, const FString& TagDevComment, ENativeGameplayTagToken)
 {
 	// TODO NDarnell To try and make sure nobody is using these during non-static init
@@ -86,6 +91,13 @@ FNativeGameplayTag::~FNativeGameplayTag()
 	}
 }
 
+#if !UE_INCLUDE_NATIVE_GAMEPLAYTAG_METADATA
+FName FNativeGameplayTag::GetModuleName() const
+{
+	return UE::NativeGameplayTag::NAME_Native;
+}
+#endif
+
 #if !UE_BUILD_SHIPPING
 
 void FNativeGameplayTag::ValidateTagRegistration() const
diff --git a/Engine/Source/Runtime/GameplayTags/Public/NativeGameplayTags.h b/Engine/Source/Runtime/GameplayTags/Public/NativeGameplayTags.h
--- a/Engine/Source/Runtime/GameplayTags/Public/NativeGameplayTags.h
+++ b/Engine/Source/Runtime/GameplayTags/Public/NativeGameplayTags.h
@@ -83,7 +83,7 @@ public:
 	FName GetModuleName() const { return ModuleName; }
 	FName GetModulePackageName() const { return ModulePackageName; }
 #else
-	FName GetModuleName() const { return NAME_Native; }
+	FName GetModuleName() const;
 #endif
 
 private:
