Changed STAT_FGameplayTag_MatchesTag to a conditional scope based on cvar GameplayTags.EnableDetailedStats.
The scope can get a bit spammy.

diff --git a/Engine/Source/Runtime/GameplayTags/Private/GameplayTagContainer.cpp b/Engine/Source/Runtime/GameplayTags/Private/GameplayTagContainer.cpp
--- a/Engine/Source/Runtime/GameplayTags/Private/GameplayTagContainer.cpp
+++ b/Engine/Source/Runtime/GameplayTags/Private/GameplayTagContainer.cpp
@@ -1078,7 +1078,7 @@ FGameplayTagContainer FGameplayTag::GetGameplayTagParents() const
 
 bool FGameplayTag::MatchesTag(const FGameplayTag& TagToCheck) const
 {
-	SCOPE_CYCLE_COUNTER(STAT_FGameplayTag_MatchesTag);
+	CONDITIONAL_SCOPE_CYCLE_COUNTER(STAT_FGameplayTag_MatchesTag, GEnableGameplayTagDetailedStats);
 
 	const FGameplayTagContainer* TagContainer = UGameplayTagsManager::Get().GetSingleTagContainer(*this);
 
