Exposing some extra make tag query functions to blueprint so blueprint can make some basic tag queries at runtime.

diff --git a/Engine/Source/Runtime/GameplayTags/Classes/BlueprintGameplayTagLibrary.h b/Engine/Source/Runtime/GameplayTags/Classes/BlueprintGameplayTagLibrary.h
--- a/Engine/Source/Runtime/GameplayTags/Classes/BlueprintGameplayTagLibrary.h
+++ b/Engine/Source/Runtime/GameplayTags/Classes/BlueprintGameplayTagLibrary.h
@@ -198,6 +198,36 @@ class GAMEPLAYTAGS_API UBlueprintGameplayTagLibrary : public UBlueprintFunctionL
 	UFUNCTION(BlueprintPure, Category = "GameplayTags", meta = (BlueprintThreadSafe))
 	static FGameplayTagQuery MakeGameplayTagQuery(FGameplayTagQuery TagQuery);
 
+	/**
+	 * Creates a literal FGameplayTagQuery with a prepopulated AnyTagsMatch expression
+	 *
+	 * @param	InTags	value to set the FGameplayTagQuery expression
+	 *
+	 * @return	The literal FGameplayTagQuery
+	 */
+	UFUNCTION(BlueprintPure, Category = "GameplayTags", meta = (BlueprintThreadSafe))
+	static FGameplayTagQuery MakeGameplayTagQuery_MatchAnyTags(const FGameplayTagContainer& InTags);
+
+	/**
+	 * Creates a literal FGameplayTagQuery with a prepopulated AllTagsMatch expression
+	 *
+	* @param	InTags	value to set the FGameplayTagQuery expression
+	 *
+	 * @return	The literal FGameplayTagQuery
+	 */
+	UFUNCTION(BlueprintPure, Category = "GameplayTags", meta = (BlueprintThreadSafe))
+	static FGameplayTagQuery MakeGameplayTagQuery_MatchAllTags(const FGameplayTagContainer& InTags);
+
+	/**
+	 * Creates a literal FGameplayTagQuery with a prepopulated NoTagsMatch expression
+	 *
+	* @param	InTags	value to set the FGameplayTagQuery expression
+	 *
+	 * @return	The literal FGameplayTagQuery
+	 */
+	UFUNCTION(BlueprintPure, Category = "GameplayTags", meta = (BlueprintThreadSafe))
+	static FGameplayTagQuery MakeGameplayTagQuery_MatchNoTags(const FGameplayTagContainer& InTags);
+	
 	/**
 	 * Check Gameplay tags in the interface has all of the specified tags in the tag container (expands to include parents of asset tags)
 	 *
diff --git a/Engine/Source/Runtime/GameplayTags/Private/BlueprintGameplayTagLibrary.cpp b/Engine/Source/Runtime/GameplayTags/Private/BlueprintGameplayTagLibrary.cpp
--- a/Engine/Source/Runtime/GameplayTags/Private/BlueprintGameplayTagLibrary.cpp
+++ b/Engine/Source/Runtime/GameplayTags/Private/BlueprintGameplayTagLibrary.cpp
@@ -178,6 +178,21 @@ FGameplayTagQuery UBlueprintGameplayTagLibrary::MakeGameplayTagQuery(FGameplayTa
 	return TagQuery;
 }
 
+FGameplayTagQuery UBlueprintGameplayTagLibrary::MakeGameplayTagQuery_MatchAnyTags(const FGameplayTagContainer& InTags)
+{
+	return FGameplayTagQuery::MakeQuery_MatchAnyTags(InTags);
+}
+
+FGameplayTagQuery UBlueprintGameplayTagLibrary::MakeGameplayTagQuery_MatchAllTags(const FGameplayTagContainer& InTags)
+{
+	return FGameplayTagQuery::MakeQuery_MatchAllTags(InTags);
+}
+
+FGameplayTagQuery UBlueprintGameplayTagLibrary::MakeGameplayTagQuery_MatchNoTags(const FGameplayTagContainer& InTags)
+{
+	return FGameplayTagQuery::MakeQuery_MatchNoTags(InTags);
+}
+
 bool UBlueprintGameplayTagLibrary::HasAllMatchingGameplayTags(TScriptInterface<IGameplayTagAssetInterface> TagContainerInterface, const FGameplayTagContainer& OtherContainer)
 {
 	if (TagContainerInterface.GetInterface() == NULL)
