Adding ability for GameplayTagQueryExpressions to be converted to json. This way a QueryExpression can be stored in other locatios rather than inside .uasset data.
This needs manually defined because we can't have recursive UProperty() structures
The conversion uses the UStruct to json serialization as much as possible to avoid writing custom serialization code.

diff --git a/Engine/Source/Runtime/GameplayTags/Classes/GameplayTagContainer.h b/Engine/Source/Runtime/GameplayTags/Classes/GameplayTagContainer.h
--- a/Engine/Source/Runtime/GameplayTags/Classes/GameplayTagContainer.h
+++ b/Engine/Source/Runtime/GameplayTags/Classes/GameplayTagContainer.h
@@ -12,6 +12,7 @@
 
 class UEditableGameplayTagQuery;
 struct FGameplayTagContainer;
+class FJsonObject;
 struct FPropertyTag;
 
 GAMEPLAYTAGS_API DECLARE_LOG_CATEGORY_EXTERN(LogGameplayTags, Log, All);
@@ -746,19 +747,16 @@ struct FGameplayTagCreationWidgetHelper
 
 /** Enumerates the list of supported query expression types. */
 UENUM()
-namespace EGameplayTagQueryExprType
+enum class EGameplayTagQueryExprType : uint8
 {
-	enum Type : int
-	{
-		Undefined = 0,
-		AnyTagsMatch,
-		AllTagsMatch,
-		NoTagsMatch,
-		AnyExprMatch,
-		AllExprMatch,
-		NoExprMatch,
-	};
-}
+	Undefined = 0,
+	AnyTagsMatch,
+	AllTagsMatch,
+	NoTagsMatch,
+	AnyExprMatch,
+	AllExprMatch,
+	NoExprMatch
+};
 
 namespace EGameplayTagQueryStreamVersion
 {
@@ -983,9 +981,11 @@ struct GAMEPLAYTAGS_API FGameplayTagQueryExpression
 	void EmitTokens(TArray<uint8>& TokenStream, TArray<FGameplayTag>& TagDictionary) const;
 
 	/** Which type of expression this is. */
-	EGameplayTagQueryExprType::Type ExprType;
+	EGameplayTagQueryExprType ExprType;
+
 	/** Expression list, for expression types that need it */
 	TArray<struct FGameplayTagQueryExpression> ExprSet;
+
 	/** Tag list, for expression types that need it */
 	TArray<FGameplayTag> TagSet;
 
@@ -999,6 +999,12 @@ struct GAMEPLAYTAGS_API FGameplayTagQueryExpression
 	{
 		return (ExprType == EGameplayTagQueryExprType::AllExprMatch) || (ExprType == EGameplayTagQueryExprType::AnyExprMatch) || (ExprType == EGameplayTagQueryExprType::NoExprMatch);
 	}
+
+	/** Converts the existing TagQueryExpression into a json object. Returns true on success */
+	bool ConvertToJsonObject(TSharedRef<FJsonObject>& OutObject) const;
+
+	/** Uses the input json object and fills out the OutQueryExpression with the data. */
+	static bool MakeFromJsonObject(const TSharedRef<FJsonObject>& InObject, FGameplayTagQueryExpression& OutQueryExpression);
 };
 
 template<>
diff --git a/Engine/Source/Runtime/GameplayTags/GameplayTags.Build.cs b/Engine/Source/Runtime/GameplayTags/GameplayTags.Build.cs
--- a/Engine/Source/Runtime/GameplayTags/GameplayTags.Build.cs
+++ b/Engine/Source/Runtime/GameplayTags/GameplayTags.Build.cs
@@ -19,7 +19,9 @@ namespace UnrealBuildTool.Rules
 			PrivateDependencyModuleNames.AddRange(
 				new string[]
 				{
-					"Projects"
+					"Projects",
+					"Json",
+					"JsonUtilities"
 				}
 			);
 
diff --git a/Engine/Source/Runtime/GameplayTags/Private/GameplayTagContainer.cpp b/Engine/Source/Runtime/GameplayTags/Private/GameplayTagContainer.cpp
--- a/Engine/Source/Runtime/GameplayTags/Private/GameplayTagContainer.cpp
+++ b/Engine/Source/Runtime/GameplayTags/Private/GameplayTagContainer.cpp
@@ -10,6 +10,7 @@
 #include "GameplayTagsManager.h"
 #include "GameplayTagsModule.h"
 #include "Misc/OutputDeviceNull.h"
+#include "JsonObjectConverter.h"
 
 #include UE_INLINE_GENERATED_CPP_BY_NAME(GameplayTagContainer)
 
@@ -222,7 +223,7 @@ void FQueryEvaluator::Read(FGameplayTagQueryExpression& E)
 
 void FQueryEvaluator::ReadExpr(FGameplayTagQueryExpression& E)
 {
-	E.ExprType = (EGameplayTagQueryExprType::Type) GetToken();
+	E.ExprType = (EGameplayTagQueryExprType) GetToken();
 	if (bReadError)
 	{
 		return;
@@ -478,7 +479,7 @@ bool FQueryEvaluator::EvalNoExprMatch(FGameplayTagContainer const& Tags, bool bS
 
 bool FQueryEvaluator::EvalExpr(FGameplayTagContainer const& Tags, bool bSkip)
 {
-	EGameplayTagQueryExprType::Type const ExprType = (EGameplayTagQueryExprType::Type) GetToken();
+	EGameplayTagQueryExprType const ExprType = (EGameplayTagQueryExprType) GetToken();
 	if (bReadError)
 	{
 		return false;
@@ -1514,7 +1515,7 @@ UEditableGameplayTagQuery* FQueryEvaluator::CreateEditableQuery()
 
 UEditableGameplayTagQueryExpression* FQueryEvaluator::ReadEditableQueryExpr(UObject* ExprOuter)
 {
-	EGameplayTagQueryExprType::Type const ExprType = (EGameplayTagQueryExprType::Type) GetToken();
+	EGameplayTagQueryExprType const ExprType = (EGameplayTagQueryExprType) GetToken();
 	if (bReadError)
 	{
 		return nullptr;
@@ -1750,7 +1751,7 @@ void UEditableGameplayTagQueryExpression::EmitExprListTokens(TArray<UEditableGam
 		else
 		{
 			// null expression
-			TokenStream.Add(EGameplayTagQueryExprType::Undefined);
+			TokenStream.Add((uint8)EGameplayTagQueryExprType::Undefined);
 			if (DebugString)
 			{
 				DebugString->Append(TEXT("undefined"));
@@ -1763,7 +1764,7 @@ void UEditableGameplayTagQueryExpression::EmitExprListTokens(TArray<UEditableGam
 
 void UEditableGameplayTagQueryExpression_AnyTagsMatch::EmitTokens(TArray<uint8>& TokenStream, TArray<FGameplayTag>& TagDictionary, FString* DebugString) const
 {
-	TokenStream.Add(EGameplayTagQueryExprType::AnyTagsMatch);
+	TokenStream.Add((uint8)EGameplayTagQueryExprType::AnyTagsMatch);
 
 	if (DebugString)
 	{
@@ -1780,7 +1781,7 @@ void UEditableGameplayTagQueryExpression_AnyTagsMatch::EmitTokens(TArray<uint8>&
 
 void UEditableGameplayTagQueryExpression_AllTagsMatch::EmitTokens(TArray<uint8>& TokenStream, TArray<FGameplayTag>& TagDictionary, FString* DebugString) const
 {
-	TokenStream.Add(EGameplayTagQueryExprType::AllTagsMatch);
+	TokenStream.Add((uint8)EGameplayTagQueryExprType::AllTagsMatch);
 
 	if (DebugString)
 	{
@@ -1797,7 +1798,7 @@ void UEditableGameplayTagQueryExpression_AllTagsMatch::EmitTokens(TArray<uint8>&
 
 void UEditableGameplayTagQueryExpression_NoTagsMatch::EmitTokens(TArray<uint8>& TokenStream, TArray<FGameplayTag>& TagDictionary, FString* DebugString) const
 {
-	TokenStream.Add(EGameplayTagQueryExprType::NoTagsMatch);
+	TokenStream.Add((uint8)EGameplayTagQueryExprType::NoTagsMatch);
 
 	if (DebugString)
 	{
@@ -1814,7 +1815,7 @@ void UEditableGameplayTagQueryExpression_NoTagsMatch::EmitTokens(TArray<uint8>&
 
 void UEditableGameplayTagQueryExpression_AnyExprMatch::EmitTokens(TArray<uint8>& TokenStream, TArray<FGameplayTag>& TagDictionary, FString* DebugString) const
 {
-	TokenStream.Add(EGameplayTagQueryExprType::AnyExprMatch);
+	TokenStream.Add((uint8)EGameplayTagQueryExprType::AnyExprMatch);
 
 	if (DebugString)
 	{
@@ -1831,7 +1832,7 @@ void UEditableGameplayTagQueryExpression_AnyExprMatch::EmitTokens(TArray<uint8>&
 
 void UEditableGameplayTagQueryExpression_AllExprMatch::EmitTokens(TArray<uint8>& TokenStream, TArray<FGameplayTag>& TagDictionary, FString* DebugString) const
 {
-	TokenStream.Add(EGameplayTagQueryExprType::AllExprMatch);
+	TokenStream.Add((uint8)EGameplayTagQueryExprType::AllExprMatch);
 
 	if (DebugString)
 	{
@@ -1848,7 +1849,7 @@ void UEditableGameplayTagQueryExpression_AllExprMatch::EmitTokens(TArray<uint8>&
 
 void UEditableGameplayTagQueryExpression_NoExprMatch::EmitTokens(TArray<uint8>& TokenStream, TArray<FGameplayTag>& TagDictionary, FString* DebugString) const
 {
-	TokenStream.Add(EGameplayTagQueryExprType::NoExprMatch);
+	TokenStream.Add((uint8)EGameplayTagQueryExprType::NoExprMatch);
 
 	if (DebugString)
 	{
@@ -1915,6 +1916,101 @@ void FGameplayTagQueryExpression::EmitTokens(TArray<uint8>& TokenStream, TArray<
 	}
 }
 
+bool FGameplayTagQueryExpression::ConvertToJsonObject(TSharedRef<FJsonObject>& OutObject) const
+{
+	TArray<TSharedPtr<FJsonValue>> TagSetArray;
+	TArray<TSharedPtr<FJsonValue>> ExprSetArray;
+	for (const FGameplayTag& GameplayTag : TagSet)
+	{
+		TSharedRef<FJsonObject> TagJsonObject = MakeShared<FJsonObject>();
+		if (FJsonObjectConverter::UStructToJsonObject(FGameplayTag::StaticStruct(), &GameplayTag, TagJsonObject))
+		{
+			TagSetArray.Add(MakeShared<FJsonValueObject>(TagJsonObject).ToSharedPtr());
+		}
+		else
+		{
+			UE_LOG(LogGameplayTags, Error, TEXT("Failed to convert GameplayTag (%s) to json"), *GameplayTag.ToString());
+			return false;
+		}
+	}
+	for (const FGameplayTagQueryExpression& TagQueryExpression : ExprSet)
+	{
+		TSharedRef<FJsonObject> ExprJsonObject = MakeShared<FJsonObject>();
+		if (TagQueryExpression.ConvertToJsonObject(ExprJsonObject)) // caution recursion
+		{
+			ExprSetArray.Add(MakeShared<FJsonValueObject>(ExprJsonObject).ToSharedPtr());
+		}
+		else
+		{
+			UE_LOG(LogGameplayTags, Error, TEXT("Failed to convert tag query expression to json"));
+			return false;
+		}
+	}
+
+	OutObject->SetStringField(TEXT("ExprType"), UEnum::GetValueAsString(ExprType));
+	OutObject->SetArrayField(TEXT("ExprSet"), ExprSetArray);
+	OutObject->SetArrayField(TEXT("TagSet"), TagSetArray);
+	return true;
+}
+
+bool FGameplayTagQueryExpression::MakeFromJsonObject(const TSharedRef<FJsonObject>& InObject, FGameplayTagQueryExpression& OutQueryExpression)
+{
+	FGameplayTagQueryExpression ReturnQueryExpression;
+
+	FString FoundExprType;
+	if (InObject->TryGetStringField(TEXT("ExprType"), FoundExprType))
+	{
+		ReturnQueryExpression.ExprType = (EGameplayTagQueryExprType)StaticEnum<EGameplayTagQueryExprType>()->GetValueByNameString(FoundExprType);
+	}
+	else
+	{
+		UE_LOG(LogGameplayTags, Error, TEXT("Failed to get ExprType from json object. Unable to make GameplayTagQueryExpression"));
+		return false;
+	}
+
+	{
+		const TArray<TSharedPtr<FJsonValue>>* TagSetArray;
+		if (InObject->TryGetArrayField(TEXT("TagSet"), TagSetArray))
+		{
+			for (const TSharedPtr<FJsonValue>& GameplayTagObject : *TagSetArray)
+			{
+				FGameplayTag GameplayTag;
+				if (FJsonObjectConverter::JsonObjectToUStruct(GameplayTagObject->AsObject().ToSharedRef(), FGameplayTag::StaticStruct(), &GameplayTag))
+				{
+					ReturnQueryExpression.TagSet.Add(GameplayTag);
+				}
+				else
+				{
+					UE_LOG(LogGameplayTags, Error, TEXT("Failed to make GameplayTag from Json"));
+					return false;
+				}
+			}
+		}
+	}
+	{
+		const TArray<TSharedPtr<FJsonValue>>* ExprSetArray;
+		if (InObject->TryGetArrayField(TEXT("ExprSet"), ExprSetArray))
+		{
+			for (const TSharedPtr<FJsonValue>& ExprJsonObject : *ExprSetArray)
+			{
+				FGameplayTagQueryExpression TagQueryExpression;
+				if (MakeFromJsonObject(ExprJsonObject->AsObject().ToSharedRef(), TagQueryExpression))
+				{
+					ReturnQueryExpression.ExprSet.Add(TagQueryExpression);
+				}
+				else
+				{
+					UE_LOG(LogGameplayTags, Error, TEXT("Failed to make GameplayTagQueryExpression from Json"));
+					return false;
+				}
+			}
+		}
+	}
+
+	OutQueryExpression = MoveTemp(ReturnQueryExpression);
+	return true;
+}
+
 #if !(UE_BUILD_SHIPPING || UE_BUILD_TEST)
 
 static void GameplayTagPrintReplicationMap()
