FGameplayTagQuery::Matches on an Empty Query now returns true by default (and warns with a more descriptive message).

diff --git a/Engine/Source/Runtime/GameplayTags/Private/GameplayTagContainer.cpp b/Engine/Source/Runtime/GameplayTags/Private/GameplayTagContainer.cpp
--- a/Engine/Source/Runtime/GameplayTags/Private/GameplayTagContainer.cpp
+++ b/Engine/Source/Runtime/GameplayTags/Private/GameplayTagContainer.cpp
@@ -1388,6 +1388,12 @@ bool FGameplayTagQuery::operator!=(const FGameplayTagQuery& Other) const
 
 bool FGameplayTagQuery::Matches(FGameplayTagContainer const& Tags) const
 {
+	if (IsEmpty())
+	{
+		UE_CLOG(!Tags.IsEmpty(), LogGameplayTags, Warning, TEXT("GameplayTagQuery %s was empty when trying to match tags %s. Defaulting to true."), *GetDescription(), *Tags.ToStringSimple());
+		return true;
+	}
+
 	FQueryEvaluator QE(*this);
 	return QE.Eval(Tags);
 }
