Optimization for UBlueprintGameplayTagLibrary::HasAllMatchingGameplayTags and DoesTagAssetInterfaceHaveTag.

Avoid having to gather gameplay tags before checking them. The default IGameplayTagAssetInterface does this, but some classes like building actors have it overridden to be smarter.

diff --git a/Engine/Source/Runtime/GameplayTags/Private/BlueprintGameplayTagLibrary.cpp b/Engine/Source/Runtime/GameplayTags/Private/BlueprintGameplayTagLibrary.cpp
--- a/Engine/Source/Runtime/GameplayTags/Private/BlueprintGameplayTagLibrary.cpp
+++ b/Engine/Source/Runtime/GameplayTags/Private/BlueprintGameplayTagLibrary.cpp
@@ -195,26 +195,20 @@ FGameplayTagQuery UBlueprintGameplayTagLibrary::MakeGameplayTagQuery_MatchNoTags
 
 bool UBlueprintGameplayTagLibrary::HasAllMatchingGameplayTags(TScriptInterface<IGameplayTagAssetInterface> TagContainerInterface, const FGameplayTagContainer& OtherContainer)
 {
-	if (TagContainerInterface.GetInterface() == NULL)
+	if (IGameplayTagAssetInterface* Interface = TagContainerInterface.GetInterface())
 	{
-		return (OtherContainer.Num() == 0);
+		return Interface->HasAllMatchingGameplayTags(OtherContainer);
 	}
-
-	FGameplayTagContainer OwnedTags;
-	TagContainerInterface->GetOwnedGameplayTags(OwnedTags);
-	return (OwnedTags.HasAll(OtherContainer));
+	return OtherContainer.IsEmpty();
 }
 
 bool UBlueprintGameplayTagLibrary::DoesTagAssetInterfaceHaveTag(TScriptInterface<IGameplayTagAssetInterface> TagContainerInterface, FGameplayTag Tag)
 {
-	if (TagContainerInterface.GetInterface() == NULL)
+	if (IGameplayTagAssetInterface* Interface = TagContainerInterface.GetInterface())
 	{
-		return false;
+		return Interface->HasMatchingGameplayTag(Tag);
 	}
-
-	FGameplayTagContainer OwnedTags;
-	TagContainerInterface->GetOwnedGameplayTags(OwnedTags);
-	return (OwnedTags.HasTag(Tag));
+	return false;
 }
 
 void UBlueprintGameplayTagLibrary::AppendGameplayTagContainers(FGameplayTagContainer& InOutTagContainer, const FGameplayTagContainer& InTagContainer)
