Fix add new tag source dialogue to not add a broken ini path to the directory search list.
This broke in CL 15716765 when it changed from searching individual files.
Also show a notification when a tag source is added and append .ini if the user forgets to

diff --git a/Engine/Plugins/Editor/GameplayTagsEditor/Source/GameplayTagsEditor/Private/GameplayTagsEditorModule.cpp b/Engine/Plugins/Editor/GameplayTagsEditor/Source/GameplayTagsEditor/Private/GameplayTagsEditorModule.cpp
--- a/Engine/Plugins/Editor/GameplayTagsEditor/Source/GameplayTagsEditor/Private/GameplayTagsEditorModule.cpp
+++ b/Engine/Plugins/Editor/GameplayTagsEditor/Source/GameplayTagsEditor/Private/GameplayTagsEditorModule.cpp
@@ -819,7 +819,20 @@ public:
 			return false;
 		}
 
-		Manager.FindOrAddTagSource(FName(*NewTagSource), EGameplayTagSourceType::TagList, RootDirToUse);
+		// Tag lists should always end with .ini
+		FName TagSourceName;
+		if (NewTagSource.EndsWith(TEXT(".ini")))
+		{
+			TagSourceName = FName(*NewTagSource);
+		}
+		else
+		{
+			TagSourceName = FName(*(NewTagSource + TEXT(".ini")));
+		}
+
+		Manager.FindOrAddTagSource(TagSourceName, EGameplayTagSourceType::TagList, RootDirToUse);
+
+		ShowNotification(FText::Format(LOCTEXT("AddTagSource", "Added {0} as a source for saving new tags"), FText::FromName(TagSourceName)), 3.0f);
 
 		IGameplayTagsModule::OnTagSettingsChanged.Broadcast();
 
diff --git a/Engine/Source/Runtime/GameplayTags/Private/GameplayTagsManager.cpp b/Engine/Source/Runtime/GameplayTags/Private/GameplayTagsManager.cpp
--- a/Engine/Source/Runtime/GameplayTags/Private/GameplayTagsManager.cpp
+++ b/Engine/Source/Runtime/GameplayTags/Private/GameplayTagsManager.cpp
@@ -1900,9 +1900,9 @@ FGameplayTagSource* UGameplayTagsManager::FindOrAddTagSource(FName TagSourceName
 		}
 		else
 		{
-			// Use custom root and make sure it gets added to the ini list for later refresh
+			// Use custom root and add the root to the search list for later refresh
 			NewSource->SourceTagList->ConfigFileName = RootDirToUse / *TagSourceName.ToString();
-			RegisteredSearchPaths.FindOrAdd(NewSource->SourceTagList->ConfigFileName);
+			RegisteredSearchPaths.FindOrAdd(RootDirToUse);
 		}
 		if (GUObjectArray.IsDisregardForGC(this))
 		{
@@ -1917,10 +1917,10 @@ FGameplayTagSource* UGameplayTagsManager::FindOrAddTagSource(FName TagSourceName
 			NewSource->SourceRestrictedTagList->ConfigFileName = FString::Printf(TEXT("%sTags/%s"), *FPaths::SourceConfigDir(), *TagSourceName.ToString());
 		}
 		else
-		{
-			// Use custom root and make sure it gets added to the ini list for later refresh
+		{			
+			// Use custom root and add the root to the search list for later refresh
 			NewSource->SourceRestrictedTagList->ConfigFileName = RootDirToUse / *TagSourceName.ToString();
-			RegisteredSearchPaths.FindOrAdd(NewSource->SourceTagList->ConfigFileName);
+			RegisteredSearchPaths.FindOrAdd(RootDirToUse);
 		}
 		if (GUObjectArray.IsDisregardForGC(this))
 		{
