[GameplayTags]
* Ran script to move dllexport from type to function/variables on code

diff --git a/Engine/Source/Runtime/GameplayTags/Classes/BlueprintGameplayTagLibrary.h b/Engine/Source/Runtime/GameplayTags/Classes/BlueprintGameplayTagLibrary.h
--- a/Engine/Source/Runtime/GameplayTags/Classes/BlueprintGameplayTagLibrary.h
+++ b/Engine/Source/Runtime/GameplayTags/Classes/BlueprintGameplayTagLibrary.h
@@ -11,8 +11,8 @@
 #include "Templates/SubclassOf.h"
 #include "BlueprintGameplayTagLibrary.generated.h"
 
-UCLASS(meta=(ScriptName="GameplayTagLibrary"))
-class GAMEPLAYTAGS_API UBlueprintGameplayTagLibrary : public UBlueprintFunctionLibrary
+UCLASS(meta=(ScriptName="GameplayTagLibrary"), MinimalAPI)
+class UBlueprintGameplayTagLibrary : public UBlueprintFunctionLibrary
 {
 	GENERATED_UCLASS_BODY()
 
@@ -26,7 +26,7 @@ class GAMEPLAYTAGS_API UBlueprintGameplayTagLibrary : public UBlueprintFunctionL
 	 * @return True if TagOne matches TagTwo
 	 */
 	UFUNCTION(BlueprintPure, Category="GameplayTags", meta = (Keywords = "DoGameplayTagsMatch", BlueprintThreadSafe))
-	static bool MatchesTag(FGameplayTag TagOne, FGameplayTag TagTwo, bool bExactMatch);
+	static GAMEPLAYTAGS_API bool MatchesTag(FGameplayTag TagOne, FGameplayTag TagTwo, bool bExactMatch);
 
 	/**
 	 * Determine if TagOne matches against any tag in OtherContainer
@@ -38,27 +38,27 @@ class GAMEPLAYTAGS_API UBlueprintGameplayTagLibrary : public UBlueprintFunctionL
 	 * @return True if TagOne matches any tags explicitly present in OtherContainer
 	 */
 	UFUNCTION(BlueprintPure, Category="GameplayTags", meta = (BlueprintThreadSafe))
-	static bool MatchesAnyTags(FGameplayTag TagOne, const FGameplayTagContainer& OtherContainer, bool bExactMatch);
+	static GAMEPLAYTAGS_API bool MatchesAnyTags(FGameplayTag TagOne, const FGameplayTagContainer& OtherContainer, bool bExactMatch);
 
 	/** Returns true if the values are equal (A == B) */
 	UFUNCTION(BlueprintPure, meta=(DisplayName="Equal (GameplayTag)", CompactNodeTitle="==", BlueprintThreadSafe), Category="GameplayTags")
-	static bool EqualEqual_GameplayTag( FGameplayTag A, FGameplayTag B );
+	static GAMEPLAYTAGS_API bool EqualEqual_GameplayTag( FGameplayTag A, FGameplayTag B );
 	
 	/** Returns true if the values are not equal (A != B) */
 	UFUNCTION(BlueprintPure, meta=(DisplayName="Not Equal (GameplayTag)", CompactNodeTitle="!=", BlueprintThreadSafe), Category="GameplayTags")
-	static bool NotEqual_GameplayTag( FGameplayTag A, FGameplayTag B );
+	static GAMEPLAYTAGS_API bool NotEqual_GameplayTag( FGameplayTag A, FGameplayTag B );
 
 	/** Returns true if the passed in gameplay tag is non-null */
 	UFUNCTION(BlueprintPure, Category = "GameplayTags", meta = (BlueprintThreadSafe))
-	static bool IsGameplayTagValid(FGameplayTag GameplayTag);
+	static GAMEPLAYTAGS_API bool IsGameplayTagValid(FGameplayTag GameplayTag);
 
 	/** Returns FName of this tag */
 	UFUNCTION(BlueprintPure, Category = "GameplayTags", meta = (BlueprintThreadSafe))
-	static FName GetTagName(const FGameplayTag& GameplayTag);
+	static GAMEPLAYTAGS_API FName GetTagName(const FGameplayTag& GameplayTag);
 
 	/** Creates a literal FGameplayTag */
 	UFUNCTION(BlueprintPure, Category = "GameplayTags", meta = (BlueprintThreadSafe))
-	static FGameplayTag MakeLiteralGameplayTag(FGameplayTag Value);
+	static GAMEPLAYTAGS_API FGameplayTag MakeLiteralGameplayTag(FGameplayTag Value);
 
 	/**
 	 * Get the number of gameplay tags in the specified container
@@ -68,7 +68,7 @@ class GAMEPLAYTAGS_API UBlueprintGameplayTagLibrary : public UBlueprintFunctionL
 	 * @return The number of tags in the specified container
 	 */
 	UFUNCTION(BlueprintPure, Category="GameplayTags", meta = (BlueprintThreadSafe))
-	static int32 GetNumGameplayTagsInContainer(const FGameplayTagContainer& TagContainer);
+	static GAMEPLAYTAGS_API int32 GetNumGameplayTagsInContainer(const FGameplayTagContainer& TagContainer);
 
 	/**
 	 * Check if the tag container has the specified tag
@@ -80,7 +80,7 @@ class GAMEPLAYTAGS_API UBlueprintGameplayTagLibrary : public UBlueprintFunctionL
 	 * @return True if the container has the specified tag, false if it does not
 	 */
 	UFUNCTION(BlueprintPure, Category="GameplayTags", meta = (Keywords = "DoesContainerHaveTag", BlueprintThreadSafe))
-	static bool HasTag(const FGameplayTagContainer& TagContainer, FGameplayTag Tag, bool bExactMatch);
+	static GAMEPLAYTAGS_API bool HasTag(const FGameplayTagContainer& TagContainer, FGameplayTag Tag, bool bExactMatch);
 
 	/**
 	 * Check if the specified tag container has ANY of the tags in the other container
@@ -92,7 +92,7 @@ class GAMEPLAYTAGS_API UBlueprintGameplayTagLibrary : public UBlueprintFunctionL
 	 * @return True if the container has ANY of the tags in the other container
 	 */
 	UFUNCTION(BlueprintPure, Category="GameplayTags", meta = (Keywords = "DoesContainerMatchAnyTagsInContainer", BlueprintThreadSafe))
-	static bool HasAnyTags(const FGameplayTagContainer& TagContainer, const FGameplayTagContainer& OtherContainer, bool bExactMatch);
+	static GAMEPLAYTAGS_API bool HasAnyTags(const FGameplayTagContainer& TagContainer, const FGameplayTagContainer& OtherContainer, bool bExactMatch);
 
 	/**
 	 * Check if the specified tag container has ALL of the tags in the other container
@@ -104,7 +104,7 @@ class GAMEPLAYTAGS_API UBlueprintGameplayTagLibrary : public UBlueprintFunctionL
 	 * @return True if the container has ALL of the tags in the other container
 	 */
 	UFUNCTION(BlueprintPure, Category="GameplayTags", meta = (Keywords = "DoesContainerMatchAllTagsInContainer", BlueprintThreadSafe))
-	static bool HasAllTags(const FGameplayTagContainer& TagContainer, const FGameplayTagContainer& OtherContainer, bool bExactMatch);
+	static GAMEPLAYTAGS_API bool HasAllTags(const FGameplayTagContainer& TagContainer, const FGameplayTagContainer& OtherContainer, bool bExactMatch);
 
 	/**
 	 * Check if the specified tag query is empty
@@ -114,7 +114,7 @@ class GAMEPLAYTAGS_API UBlueprintGameplayTagLibrary : public UBlueprintFunctionL
 	 * @return True if the query is empty, false otherwise.
 	 */
 	UFUNCTION(BlueprintPure, Category = "GameplayTags", meta = (BlueprintThreadSafe))
-	static bool IsTagQueryEmpty(const FGameplayTagQuery& TagQuery);
+	static GAMEPLAYTAGS_API bool IsTagQueryEmpty(const FGameplayTagQuery& TagQuery);
 
 	/**
 	 * Check if the specified tag container matches the given Tag Query
@@ -125,7 +125,7 @@ class GAMEPLAYTAGS_API UBlueprintGameplayTagLibrary : public UBlueprintFunctionL
 	 * @return True if the container matches the query, false otherwise.
 	 */
 	UFUNCTION(BlueprintPure, Category = "GameplayTags", meta = (BlueprintThreadSafe))
-	static bool DoesContainerMatchTagQuery(const FGameplayTagContainer& TagContainer, const FGameplayTagQuery& TagQuery);
+	static GAMEPLAYTAGS_API bool DoesContainerMatchTagQuery(const FGameplayTagContainer& TagContainer, const FGameplayTagQuery& TagQuery);
 
 	/**
 	 * Get an array of all actors of a specific class (or subclass of that class) which match the specified gameplay tag query.
@@ -135,7 +135,7 @@ class GAMEPLAYTAGS_API UBlueprintGameplayTagLibrary : public UBlueprintFunctionL
 	 * 
 	 */
 	UFUNCTION(BlueprintCallable, Category="GameplayTags",  meta=(WorldContext="WorldContextObject", DeterminesOutputType="ActorClass", DynamicOutputParam="OutActors"))
-	static void GetAllActorsOfClassMatchingTagQuery(UObject* WorldContextObject, TSubclassOf<AActor> ActorClass, const FGameplayTagQuery& GameplayTagQuery, TArray<AActor*>& OutActors);
+	static GAMEPLAYTAGS_API void GetAllActorsOfClassMatchingTagQuery(UObject* WorldContextObject, TSubclassOf<AActor> ActorClass, const FGameplayTagQuery& GameplayTagQuery, TArray<AActor*>& OutActors);
 
 	/**
 	 * Adds a single tag to the passed in tag container
@@ -144,7 +144,7 @@ class GAMEPLAYTAGS_API UBlueprintGameplayTagLibrary : public UBlueprintFunctionL
 	 * @param Tag					The tag to add to the container
 	 */
 	UFUNCTION(BlueprintCallable, Category = "GameplayTags", meta = (BlueprintThreadSafe))
-	static void AddGameplayTag(UPARAM(ref) FGameplayTagContainer& TagContainer, FGameplayTag Tag);
+	static GAMEPLAYTAGS_API void AddGameplayTag(UPARAM(ref) FGameplayTagContainer& TagContainer, FGameplayTag Tag);
 
 	/**
 	 * Remove a single tag from the passed in tag container, returns true if found
@@ -153,7 +153,7 @@ class GAMEPLAYTAGS_API UBlueprintGameplayTagLibrary : public UBlueprintFunctionL
 	 * @param Tag					The tag to add to the container
 	 */
 	UFUNCTION(BlueprintCallable, Category = "GameplayTags", meta = (BlueprintThreadSafe))
-	static bool RemoveGameplayTag(UPARAM(ref) FGameplayTagContainer& TagContainer, FGameplayTag Tag);
+	static GAMEPLAYTAGS_API bool RemoveGameplayTag(UPARAM(ref) FGameplayTagContainer& TagContainer, FGameplayTag Tag);
 
 	/**
 	 * Appends all tags in the InTagContainer to InOutTagContainer
@@ -162,31 +162,31 @@ class GAMEPLAYTAGS_API UBlueprintGameplayTagLibrary : public UBlueprintFunctionL
 	 * @param InTagContainer		The container to append.
 	 */
 	UFUNCTION(BlueprintCallable, Category = "GameplayTags", meta = (BlueprintThreadSafe))
-	static void AppendGameplayTagContainers(UPARAM(ref) FGameplayTagContainer& InOutTagContainer, const FGameplayTagContainer& InTagContainer);
+	static GAMEPLAYTAGS_API void AppendGameplayTagContainers(UPARAM(ref) FGameplayTagContainer& InOutTagContainer, const FGameplayTagContainer& InTagContainer);
 
 	/** Returns true if the values are equal (A == B) */
 	UFUNCTION(BlueprintPure, meta=(DisplayName="Equal (GameplayTagContainer)", CompactNodeTitle="==", BlueprintThreadSafe), Category="GameplayTags")
-	static bool EqualEqual_GameplayTagContainer( const FGameplayTagContainer& A, const FGameplayTagContainer& B );
+	static GAMEPLAYTAGS_API bool EqualEqual_GameplayTagContainer( const FGameplayTagContainer& A, const FGameplayTagContainer& B );
 	
 	/** Returns true if the values are not equal (A != B) */
 	UFUNCTION(BlueprintPure, meta=(DisplayName="Not Equal (GameplayTagContainer)", CompactNodeTitle="!=", BlueprintThreadSafe), Category="GameplayTags")
-	static bool NotEqual_GameplayTagContainer( const FGameplayTagContainer& A, const FGameplayTagContainer& B );
+	static GAMEPLAYTAGS_API bool NotEqual_GameplayTagContainer( const FGameplayTagContainer& A, const FGameplayTagContainer& B );
 
 	/** Creates a literal FGameplayTagContainer */
 	UFUNCTION(BlueprintPure, Category = "GameplayTags", meta = (BlueprintThreadSafe))
-	static FGameplayTagContainer MakeLiteralGameplayTagContainer(FGameplayTagContainer Value);
+	static GAMEPLAYTAGS_API FGameplayTagContainer MakeLiteralGameplayTagContainer(FGameplayTagContainer Value);
 
 	/** Creates a FGameplayTagContainer from the array of passed in tags */
 	UFUNCTION(BlueprintPure, Category = "GameplayTags", meta = (BlueprintThreadSafe))
-	static FGameplayTagContainer MakeGameplayTagContainerFromArray(const TArray<FGameplayTag>& GameplayTags);
+	static GAMEPLAYTAGS_API FGameplayTagContainer MakeGameplayTagContainerFromArray(const TArray<FGameplayTag>& GameplayTags);
 
 	/** Creates a FGameplayTagContainer containing a single tag */
 	UFUNCTION(BlueprintPure, Category = "GameplayTags", meta = (BlueprintThreadSafe))
-	static FGameplayTagContainer MakeGameplayTagContainerFromTag(FGameplayTag SingleTag);
+	static GAMEPLAYTAGS_API FGameplayTagContainer MakeGameplayTagContainerFromTag(FGameplayTag SingleTag);
 
 	/** Breaks tag container into explicit array of tags */
 	UFUNCTION(BlueprintPure, Category = "GameplayTags", meta = (BlueprintThreadSafe))
-	static void BreakGameplayTagContainer(const FGameplayTagContainer& GameplayTagContainer, TArray<FGameplayTag>& GameplayTags);
+	static GAMEPLAYTAGS_API void BreakGameplayTagContainer(const FGameplayTagContainer& GameplayTagContainer, TArray<FGameplayTag>& GameplayTags);
 
 	/**
 	 * Creates a literal FGameplayTagQuery
@@ -196,7 +196,7 @@ class GAMEPLAYTAGS_API UBlueprintGameplayTagLibrary : public UBlueprintFunctionL
 	 * @return	The literal FGameplayTagQuery
 	 */
 	UFUNCTION(BlueprintPure, Category = "GameplayTags", meta = (BlueprintThreadSafe))
-	static FGameplayTagQuery MakeGameplayTagQuery(FGameplayTagQuery TagQuery);
+	static GAMEPLAYTAGS_API FGameplayTagQuery MakeGameplayTagQuery(FGameplayTagQuery TagQuery);
 
 	/**
 	 * Creates a literal FGameplayTagQuery with a prepopulated AnyTagsMatch expression
@@ -206,7 +206,7 @@ class GAMEPLAYTAGS_API UBlueprintGameplayTagLibrary : public UBlueprintFunctionL
 	 * @return	The literal FGameplayTagQuery
 	 */
 	UFUNCTION(BlueprintPure, Category = "GameplayTags", meta = (BlueprintThreadSafe))
-	static FGameplayTagQuery MakeGameplayTagQuery_MatchAnyTags(const FGameplayTagContainer& InTags);
+	static GAMEPLAYTAGS_API FGameplayTagQuery MakeGameplayTagQuery_MatchAnyTags(const FGameplayTagContainer& InTags);
 
 	/**
 	 * Creates a literal FGameplayTagQuery with a prepopulated AllTagsMatch expression
@@ -216,7 +216,7 @@ class GAMEPLAYTAGS_API UBlueprintGameplayTagLibrary : public UBlueprintFunctionL
 	 * @return	The literal FGameplayTagQuery
 	 */
 	UFUNCTION(BlueprintPure, Category = "GameplayTags", meta = (BlueprintThreadSafe))
-	static FGameplayTagQuery MakeGameplayTagQuery_MatchAllTags(const FGameplayTagContainer& InTags);
+	static GAMEPLAYTAGS_API FGameplayTagQuery MakeGameplayTagQuery_MatchAllTags(const FGameplayTagContainer& InTags);
 
 	/**
 	 * Creates a literal FGameplayTagQuery with a prepopulated NoTagsMatch expression
@@ -226,7 +226,7 @@ class GAMEPLAYTAGS_API UBlueprintGameplayTagLibrary : public UBlueprintFunctionL
 	 * @return	The literal FGameplayTagQuery
 	 */
 	UFUNCTION(BlueprintPure, Category = "GameplayTags", meta = (BlueprintThreadSafe))
-	static FGameplayTagQuery MakeGameplayTagQuery_MatchNoTags(const FGameplayTagContainer& InTags);
+	static GAMEPLAYTAGS_API FGameplayTagQuery MakeGameplayTagQuery_MatchNoTags(const FGameplayTagContainer& InTags);
 	
 	/**
 	 * Check Gameplay tags in the interface has all of the specified tags in the tag container (expands to include parents of asset tags)
@@ -237,7 +237,7 @@ class GAMEPLAYTAGS_API UBlueprintGameplayTagLibrary : public UBlueprintFunctionL
 	 * @return True if the tagcontainer in the interface has all the tags inside the container.
 	 */
 	UFUNCTION(BlueprintPure, meta = (BlueprintInternalUseOnly = "TRUE"))
-	static bool HasAllMatchingGameplayTags(TScriptInterface<IGameplayTagAssetInterface> TagContainerInterface, const FGameplayTagContainer& OtherContainer);
+	static GAMEPLAYTAGS_API bool HasAllMatchingGameplayTags(TScriptInterface<IGameplayTagAssetInterface> TagContainerInterface, const FGameplayTagContainer& OtherContainer);
 
 	/**
 	 * Check if the specified tag container has the specified tag, using the specified tag matching types
@@ -248,15 +248,15 @@ class GAMEPLAYTAGS_API UBlueprintGameplayTagLibrary : public UBlueprintFunctionL
 	 * @return True if the container has the specified tag, false if it does not
 	 */
 	UFUNCTION(BlueprintPure, meta = (BlueprintInternalUseOnly = "TRUE"))
-	static bool DoesTagAssetInterfaceHaveTag(TScriptInterface<IGameplayTagAssetInterface> TagContainerInterface, FGameplayTag Tag);
+	static GAMEPLAYTAGS_API bool DoesTagAssetInterfaceHaveTag(TScriptInterface<IGameplayTagAssetInterface> TagContainerInterface, FGameplayTag Tag);
 
 	/** Checks if a gameplay tag's name and a string are not equal to one another */
 	UFUNCTION(BlueprintPure, Category = PinOptions, meta = (BlueprintInternalUseOnly = "TRUE"))
-	static bool NotEqual_TagTag(FGameplayTag A, FString B);
+	static GAMEPLAYTAGS_API bool NotEqual_TagTag(FGameplayTag A, FString B);
 
 	/** Checks if a gameplay tag containers's name and a string are not equal to one another */
 	UFUNCTION(BlueprintPure, Category = PinOptions, meta = (BlueprintInternalUseOnly = "TRUE"))
-	static bool NotEqual_TagContainerTagContainer(FGameplayTagContainer A, FString B);
+	static GAMEPLAYTAGS_API bool NotEqual_TagContainerTagContainer(FGameplayTagContainer A, FString B);
 	
 	/**
 	 * Returns an FString listing all of the gameplay tags in the tag container for debugging purposes.
@@ -264,7 +264,7 @@ class GAMEPLAYTAGS_API UBlueprintGameplayTagLibrary : public UBlueprintFunctionL
 	 * @param TagContainer	The tag container to get the debug string from.
 	 */
 	UFUNCTION(BlueprintPure, Category = "GameplayTags", meta = (BlueprintThreadSafe))
-	static FString GetDebugStringFromGameplayTagContainer(const FGameplayTagContainer& TagContainer);
+	static GAMEPLAYTAGS_API FString GetDebugStringFromGameplayTagContainer(const FGameplayTagContainer& TagContainer);
 
 	/**
 	 * Returns an FString representation of a gameplay tag for debugging purposes.
@@ -272,6 +272,6 @@ class GAMEPLAYTAGS_API UBlueprintGameplayTagLibrary : public UBlueprintFunctionL
 	 * @param GameplayTag	The tag to get the debug string from.
 	 */
 	UFUNCTION(BlueprintPure, Category = "GameplayTags", meta = (BlueprintThreadSafe))
-	static FString GetDebugStringFromGameplayTag(FGameplayTag GameplayTag);
+	static GAMEPLAYTAGS_API FString GetDebugStringFromGameplayTag(FGameplayTag GameplayTag);
 
 };
diff --git a/Engine/Source/Runtime/GameplayTags/Classes/GameplayTagAssetInterface.h b/Engine/Source/Runtime/GameplayTags/Classes/GameplayTagAssetInterface.h
--- a/Engine/Source/Runtime/GameplayTags/Classes/GameplayTagAssetInterface.h
+++ b/Engine/Source/Runtime/GameplayTags/Classes/GameplayTagAssetInterface.h
@@ -16,7 +16,7 @@ class UGameplayTagAssetInterface : public UInterface
 	GENERATED_UINTERFACE_BODY()
 };
 
-class GAMEPLAYTAGS_API IGameplayTagAssetInterface
+class IGameplayTagAssetInterface
 {
 	GENERATED_IINTERFACE_BODY()
 
@@ -36,7 +36,7 @@ class GAMEPLAYTAGS_API IGameplayTagAssetInterface
 	 * @return True if the asset has a gameplay tag that matches, false if not
 	 */
 	UFUNCTION(BlueprintCallable, Category=GameplayTags)
-	virtual bool HasMatchingGameplayTag(FGameplayTag TagToCheck) const;
+	GAMEPLAYTAGS_API virtual bool HasMatchingGameplayTag(FGameplayTag TagToCheck) const;
 
 	/**
 	 * Check if the asset has gameplay tags that matches against all of the specified tags (expands to include parents of asset tags)
@@ -46,7 +46,7 @@ class GAMEPLAYTAGS_API IGameplayTagAssetInterface
 	 * @return True if the asset has matches all of the gameplay tags, will be true if container is empty
 	 */
 	UFUNCTION(BlueprintCallable, Category=GameplayTags)
-	virtual bool HasAllMatchingGameplayTags(const FGameplayTagContainer& TagContainer) const;
+	GAMEPLAYTAGS_API virtual bool HasAllMatchingGameplayTags(const FGameplayTagContainer& TagContainer) const;
 
 	/**
 	 * Check if the asset has gameplay tags that matches against any of the specified tags (expands to include parents of asset tags)
@@ -56,6 +56,6 @@ class GAMEPLAYTAGS_API IGameplayTagAssetInterface
 	 * @return True if the asset has matches any of the gameplay tags, will be false if container is empty
 	 */
 	UFUNCTION(BlueprintCallable, Category=GameplayTags)
-	virtual bool HasAnyMatchingGameplayTags(const FGameplayTagContainer& TagContainer) const;
+	GAMEPLAYTAGS_API virtual bool HasAnyMatchingGameplayTags(const FGameplayTagContainer& TagContainer) const;
 };
 
diff --git a/Engine/Source/Runtime/GameplayTags/Classes/GameplayTagContainer.h b/Engine/Source/Runtime/GameplayTags/Classes/GameplayTagContainer.h
--- a/Engine/Source/Runtime/GameplayTags/Classes/GameplayTagContainer.h
+++ b/Engine/Source/Runtime/GameplayTags/Classes/GameplayTagContainer.h
@@ -48,7 +48,7 @@ typedef uint16 FGameplayTagNetIndex;
  * You can filter the gameplay tags displayed in the editor using, meta = (Categories = "Tag1.Tag2.Tag3"))
  */
 USTRUCT(BlueprintType, meta = (HasNativeMake = "/Script/GameplayTags.BlueprintGameplayTagLibrary.MakeLiteralGameplayTag", HasNativeBreak = "/Script/GameplayTags.BlueprintGameplayTagLibrary.GetTagName", DisableSplitPin))
-struct GAMEPLAYTAGS_API FGameplayTag
+struct FGameplayTag
 {
 	GENERATED_USTRUCT_BODY()
 
@@ -64,7 +64,7 @@ struct GAMEPLAYTAGS_API FGameplayTag
 	 * @param ErrorIfNotfound: ensure() that tag exists.
 	 * @return Will return the corresponding FGameplayTag or an empty one if not found.
 	 */
-	static FGameplayTag RequestGameplayTag(const FName& TagName, bool ErrorIfNotFound=true);
+	static GAMEPLAYTAGS_API FGameplayTag RequestGameplayTag(const FName& TagName, bool ErrorIfNotFound=true);
 
 	/** 
 	 * Returns true if this is a valid gameplay tag string (foo.bar.baz). If false, it will fill 
@@ -73,7 +73,7 @@ struct GAMEPLAYTAGS_API FGameplayTag
 	 * @param OutFixedString If non-null and string invalid, will attempt to fix. Will be empty if no fix is possible
 	 * @return True if this can be added to the tag dictionary, false if there's a syntax error
 	 */
-	static bool IsValidGameplayTagString(const FString& TagString, FText* OutError = nullptr, FString* OutFixedString = nullptr);
+	static GAMEPLAYTAGS_API bool IsValidGameplayTagString(const FString& TagString, FText* OutError = nullptr, FString* OutFixedString = nullptr);
 
 	/** Operators */
 	FORCEINLINE bool operator==(FGameplayTag const& Other) const
@@ -98,7 +98,7 @@ struct GAMEPLAYTAGS_API FGameplayTag
 	 * 
 	 * @return True if this tag matches TagToCheck
 	 */
-	bool MatchesTag(const FGameplayTag& TagToCheck) const;
+	GAMEPLAYTAGS_API bool MatchesTag(const FGameplayTag& TagToCheck) const;
 
 	/**
 	 * Determine if TagToCheck is valid and exactly matches this tag
@@ -124,7 +124,7 @@ struct GAMEPLAYTAGS_API FGameplayTag
 	 *
 	 * @return The depth of the match, higher means they are closer to an exact match
 	 */
-	int32 MatchesTagDepth(const FGameplayTag& TagToCheck) const;
+	GAMEPLAYTAGS_API int32 MatchesTagDepth(const FGameplayTag& TagToCheck) const;
 
 	/**
 	 * Checks if this tag matches ANY of the tags in the specified container, also checks against our parent tags
@@ -133,7 +133,7 @@ struct GAMEPLAYTAGS_API FGameplayTag
 	 *
 	 * @return True if this tag matches ANY of the tags of in ContainerToCheck
 	 */
-	bool MatchesAny(const FGameplayTagContainer& ContainerToCheck) const;
+	GAMEPLAYTAGS_API bool MatchesAny(const FGameplayTagContainer& ContainerToCheck) const;
 
 	/**
 	 * Checks if this tag matches ANY of the tags in the specified container, only allowing exact matches
@@ -142,7 +142,7 @@ struct GAMEPLAYTAGS_API FGameplayTag
 	 *
 	 * @return True if this tag matches ANY of the tags of in ContainerToCheck exactly
 	 */
-	bool MatchesAnyExact(const FGameplayTagContainer& ContainerToCheck) const;
+	GAMEPLAYTAGS_API bool MatchesAnyExact(const FGameplayTagContainer& ContainerToCheck) const;
 
 	/** Returns whether the tag is valid or not; Invalid tags are set to NAME_None and do not exist in the game-specific global dictionary */
 	FORCEINLINE bool IsValid() const
@@ -151,13 +151,13 @@ struct GAMEPLAYTAGS_API FGameplayTag
 	}
 
 	/** Returns reference to a GameplayTagContainer containing only this tag */
-	const FGameplayTagContainer& GetSingleTagContainer() const;
+	GAMEPLAYTAGS_API const FGameplayTagContainer& GetSingleTagContainer() const;
 
 	/** Returns direct parent GameplayTag of this GameplayTag, calling on x.y will return x */
-	FGameplayTag RequestDirectParent() const;
+	GAMEPLAYTAGS_API FGameplayTag RequestDirectParent() const;
 
 	/** Returns a new container explicitly containing the tags of this tag */
-	FGameplayTagContainer GetGameplayTagParents() const;
+	GAMEPLAYTAGS_API FGameplayTagContainer GetGameplayTagParents() const;
 
 	/** Used so we can have a TMap of this struct */
 	FORCEINLINE friend uint32 GetTypeHash(const FGameplayTag& Tag)
@@ -183,28 +183,28 @@ struct GAMEPLAYTAGS_API FGameplayTag
 	}
 
 	/** Overridden for fast serialize */
-	bool NetSerialize(FArchive& Ar, class UPackageMap* Map, bool& bOutSuccess);
+	GAMEPLAYTAGS_API bool NetSerialize(FArchive& Ar, class UPackageMap* Map, bool& bOutSuccess);
 
 	/** Handles fixup and errors. This is only called when not serializing a full FGameplayTagContainer */
-	void PostSerialize(const FArchive& Ar);
-	bool NetSerialize_Packed(FArchive& Ar, class UPackageMap* Map, bool& bOutSuccess);
+	GAMEPLAYTAGS_API void PostSerialize(const FArchive& Ar);
+	GAMEPLAYTAGS_API bool NetSerialize_Packed(FArchive& Ar, class UPackageMap* Map, bool& bOutSuccess);
 
 	/** Used to upgrade a Name property to a GameplayTag struct property */
-	bool SerializeFromMismatchedTag(const FPropertyTag& Tag, FStructuredArchive::FSlot Slot);
+	GAMEPLAYTAGS_API bool SerializeFromMismatchedTag(const FPropertyTag& Tag, FStructuredArchive::FSlot Slot);
 
 	/** Sets from a ImportText string, used in asset registry */
-	void FromExportString(const FString& ExportString, int32 PortFlags = 0);
+	GAMEPLAYTAGS_API void FromExportString(const FString& ExportString, int32 PortFlags = 0);
 
 	/** Handles importing tag strings without (TagName=) in it */
-	bool ImportTextItem(const TCHAR*& Buffer, int32 PortFlags, UObject* Parent, FOutputDevice* ErrorText);
+	GAMEPLAYTAGS_API bool ImportTextItem(const TCHAR*& Buffer, int32 PortFlags, UObject* Parent, FOutputDevice* ErrorText);
 
 	/** An empty Gameplay Tag */
-	static const FGameplayTag EmptyTag;
+	static GAMEPLAYTAGS_API const FGameplayTag EmptyTag;
 
 protected:
 
 	/** Intentionally private so only the tag manager can use */
-	explicit FGameplayTag(const FName& InTagName);
+	GAMEPLAYTAGS_API explicit FGameplayTag(const FName& InTagName);
 
 	/** This Tags Name */
 	UPROPERTY(VisibleAnywhere, Category = GameplayTags, SaveGame)
@@ -232,7 +232,7 @@ struct TStructOpsTypeTraits< FGameplayTag > : public TStructOpsTypeTraitsBase2<
 
 /** A Tag Container holds a collection of FGameplayTags, tags are included explicitly by adding them, and implicitly from adding child tags */
 USTRUCT(BlueprintType, meta = (HasNativeMake = "/Script/GameplayTags.BlueprintGameplayTagLibrary.MakeGameplayTagContainerFromArray", HasNativeBreak = "/Script/GameplayTags.BlueprintGameplayTagLibrary.BreakGameplayTagContainer"))
-struct GAMEPLAYTAGS_API FGameplayTagContainer
+struct FGameplayTagContainer
 {
 	GENERATED_USTRUCT_BODY()
 
@@ -274,10 +274,10 @@ struct GAMEPLAYTAGS_API FGameplayTagContainer
 	}
 
 	/** Assignment/Equality operators */
-	FGameplayTagContainer& operator=(FGameplayTagContainer const& Other);
-	FGameplayTagContainer& operator=(FGameplayTagContainer&& Other);
-	bool operator==(FGameplayTagContainer const& Other) const;
-	bool operator!=(FGameplayTagContainer const& Other) const;
+	GAMEPLAYTAGS_API FGameplayTagContainer& operator=(FGameplayTagContainer const& Other);
+	GAMEPLAYTAGS_API FGameplayTagContainer& operator=(FGameplayTagContainer&& Other);
+	GAMEPLAYTAGS_API bool operator==(FGameplayTagContainer const& Other) const;
+	GAMEPLAYTAGS_API bool operator!=(FGameplayTagContainer const& Other) const;
 
 	/**
 	 * Determine if TagToCheck is present in this container, also checking against parent tags
@@ -424,7 +424,7 @@ struct GAMEPLAYTAGS_API FGameplayTagContainer
 	}
 
 	/** Returns a new container explicitly containing the tags of this container and all of their parent tags */
-	FGameplayTagContainer GetGameplayTagParents() const;
+	GAMEPLAYTAGS_API FGameplayTagContainer GetGameplayTagParents() const;
 
 	/**
 	 * Returns a filtered version of this container, returns all tags that match against any of the tags in OtherContainer, expanding parents
@@ -433,7 +433,7 @@ struct GAMEPLAYTAGS_API FGameplayTagContainer
 	 *
 	 * @return A FGameplayTagContainer containing the filtered tags
 	 */
-	FGameplayTagContainer Filter(const FGameplayTagContainer& OtherContainer) const;
+	GAMEPLAYTAGS_API FGameplayTagContainer Filter(const FGameplayTagContainer& OtherContainer) const;
 
 	/**
 	 * Returns a filtered version of this container, returns all tags that match exactly one in OtherContainer
@@ -442,7 +442,7 @@ struct GAMEPLAYTAGS_API FGameplayTagContainer
 	 *
 	 * @return A FGameplayTagContainer containing the filtered tags
 	 */
-	FGameplayTagContainer FilterExact(const FGameplayTagContainer& OtherContainer) const;
+	GAMEPLAYTAGS_API FGameplayTagContainer FilterExact(const FGameplayTagContainer& OtherContainer) const;
 
 	/** 
 	 * Checks if this container matches the given query.
@@ -451,7 +451,7 @@ struct GAMEPLAYTAGS_API FGameplayTagContainer
 	 *
 	 * @return True if this container matches the query, false otherwise.
 	 */
-	bool MatchesQuery(const struct FGameplayTagQuery& Query) const;
+	GAMEPLAYTAGS_API bool MatchesQuery(const struct FGameplayTagQuery& Query) const;
 
 	/** 
 	 * Adds all the tags from one container to this container 
@@ -459,7 +459,7 @@ struct GAMEPLAYTAGS_API FGameplayTagContainer
 	 *
 	 * @param Other TagContainer that has the tags you want to add to this container 
 	 */
-	void AppendTags(FGameplayTagContainer const& Other);
+	GAMEPLAYTAGS_API void AppendTags(FGameplayTagContainer const& Other);
 
 	/** 
 	 * Adds all the tags that match between the two specified containers to this container.  WARNING: This matches any
@@ -474,14 +474,14 @@ struct GAMEPLAYTAGS_API FGameplayTagContainer
 	 * @param OtherA TagContainer that has the matching tags you want to add to this container, these tags have their parents expanded
 	 * @param OtherB TagContainer used to check for matching tags.  If the tag matches on any parent, it counts as a match.
 	 */
-	void AppendMatchingTags(FGameplayTagContainer const& OtherA, FGameplayTagContainer const& OtherB);
+	GAMEPLAYTAGS_API void AppendMatchingTags(FGameplayTagContainer const& OtherA, FGameplayTagContainer const& OtherB);
 
 	/**
 	 * Add the specified tag to the container
 	 *
 	 * @param TagToAdd Tag to add to the container
 	 */
-	void AddTag(const FGameplayTag& TagToAdd);
+	GAMEPLAYTAGS_API void AddTag(const FGameplayTag& TagToAdd);
 
 	/**
 	 * Add the specified tag to the container without checking for uniqueness
@@ -490,7 +490,7 @@ struct GAMEPLAYTAGS_API FGameplayTagContainer
 	 * 
 	 * Useful when building container from another data struct (TMap for example)
 	 */
-	void AddTagFast(const FGameplayTag& TagToAdd);
+	GAMEPLAYTAGS_API void AddTagFast(const FGameplayTag& TagToAdd);
 
 	/**
 	 * Adds a tag to the container and removes any direct parents, wont add if child already exists
@@ -499,7 +499,7 @@ struct GAMEPLAYTAGS_API FGameplayTagContainer
 	 * 
 	 * @return True if tag was added
 	 */
-	bool AddLeafTag(const FGameplayTag& TagToAdd);
+	GAMEPLAYTAGS_API bool AddLeafTag(const FGameplayTag& TagToAdd);
 
 	/**
 	 * Tag to remove from the container
@@ -507,44 +507,44 @@ struct GAMEPLAYTAGS_API FGameplayTagContainer
 	 * @param TagToRemove		Tag to remove from the container
 	 * @param bDeferParentTags	Skip calling FillParentTags for performance (must be handled by calling code)
 	 */
-	bool RemoveTag(const FGameplayTag& TagToRemove, bool bDeferParentTags=false);
+	GAMEPLAYTAGS_API bool RemoveTag(const FGameplayTag& TagToRemove, bool bDeferParentTags=false);
 
 	/**
 	 * Removes all tags in TagsToRemove from this container
 	 *
 	 * @param TagsToRemove	Tags to remove from the container
 	 */
-	void RemoveTags(const FGameplayTagContainer& TagsToRemove);
+	GAMEPLAYTAGS_API void RemoveTags(const FGameplayTagContainer& TagsToRemove);
 
 	/** Remove all tags from the container. Will maintain slack by default */
-	void Reset(int32 Slack = 0);
+	GAMEPLAYTAGS_API void Reset(int32 Slack = 0);
 	
 	/** Serialize the tag container */
-	bool Serialize(FStructuredArchive::FSlot Slot);
+	GAMEPLAYTAGS_API bool Serialize(FStructuredArchive::FSlot Slot);
 
 	/** Efficient network serialize, takes advantage of the dictionary */
-	bool NetSerialize(FArchive& Ar, class UPackageMap* Map, bool& bOutSuccess);
+	GAMEPLAYTAGS_API bool NetSerialize(FArchive& Ar, class UPackageMap* Map, bool& bOutSuccess);
 
 	/** Handles fixup after importing from text */
-	bool ImportTextItem(const TCHAR*& Buffer, int32 PortFlags, UObject* Parent, FOutputDevice* ErrorText);
+	GAMEPLAYTAGS_API bool ImportTextItem(const TCHAR*& Buffer, int32 PortFlags, UObject* Parent, FOutputDevice* ErrorText);
 
 	/** Fill in the ParentTags array and any other transient parameters */
-	void PostScriptConstruct();
+	GAMEPLAYTAGS_API void PostScriptConstruct();
 
 	/** Returns string version of container in ImportText format */
-	FString ToString() const;
+	GAMEPLAYTAGS_API FString ToString() const;
 
 	/** Sets from a ImportText string, used in asset registry */
-	void FromExportString(const FString& ExportString, int32 PortFlags = 0);
+	GAMEPLAYTAGS_API void FromExportString(const FString& ExportString, int32 PortFlags = 0);
 
 	/** Returns abbreviated human readable Tag list without parens or property names. If bQuoted is true it will quote each tag */
-	FString ToStringSimple(bool bQuoted = false) const;
+	GAMEPLAYTAGS_API FString ToStringSimple(bool bQuoted = false) const;
 
 	/** Returns abbreviated human readable Tag list without parens or property names, but will limit each string to specified len.  This is to get around output restrictions*/
-	TArray<FString> ToStringsMaxLen(int32 MaxLen) const;
+	GAMEPLAYTAGS_API TArray<FString> ToStringsMaxLen(int32 MaxLen) const;
 
 	/** Returns human readable description of what match is being looked for on the readable tag list. */
-	FText ToMatchingText(EGameplayContainerMatchType MatchType, bool bInvertCondition) const;
+	GAMEPLAYTAGS_API FText ToMatchingText(EGameplayContainerMatchType MatchType, bool bInvertCondition) const;
 
 	/** Gets the explicit list of gameplay tags */
 	void GetGameplayTagArray(TArray<FGameplayTag>& InOutGameplayTags) const
@@ -553,7 +553,7 @@ struct GAMEPLAYTAGS_API FGameplayTagContainer
 	}
 
 	/** Gets the explicit list of gameplay tags */
-	const TArray<FGameplayTag>& GetGameplayTagArray() const;
+	GAMEPLAYTAGS_API const TArray<FGameplayTag>& GetGameplayTagArray() const;
 
 	/** Creates a const iterator for the contents of this array */
 	TArray<FGameplayTag>::TConstIterator CreateConstIterator() const
@@ -586,10 +586,10 @@ struct GAMEPLAYTAGS_API FGameplayTagContainer
 	}
 
 	/** Fills in ParentTags from GameplayTags */
-	void FillParentTags();
+	GAMEPLAYTAGS_API void FillParentTags();
 
 	/** An empty Gameplay Tag Container */
-	static const FGameplayTagContainer EmptyContainer;
+	static GAMEPLAYTAGS_API const FGameplayTagContainer EmptyContainer;
 
 PRAGMA_DISABLE_DEPRECATION_WARNINGS
 
@@ -616,7 +616,7 @@ PRAGMA_DISABLE_DEPRECATION_WARNINGS
 	}
 
 	UE_DEPRECATED(5.0, "Deprecated in favor of HasTag or HasTagExact")
-	bool ComplexHasTag(FGameplayTag const& TagToCheck, TEnumAsByte<EGameplayTagMatchType::Type> TagMatchType, TEnumAsByte<EGameplayTagMatchType::Type> TagToCheckMatchType) const;
+	GAMEPLAYTAGS_API bool ComplexHasTag(FGameplayTag const& TagToCheck, TEnumAsByte<EGameplayTagMatchType::Type> TagMatchType, TEnumAsByte<EGameplayTagMatchType::Type> TagToCheckMatchType) const;
 
 	UE_DEPRECATED(5.0, "Deprecated in favor of HasAll and related functions")
 	FORCEINLINE_DEBUGGABLE bool DoesTagContainerMatch(const FGameplayTagContainer& OtherContainer, TEnumAsByte<EGameplayTagMatchType::Type> TagMatchType, TEnumAsByte<EGameplayTagMatchType::Type> OtherTagMatchType, EGameplayContainerMatchType ContainerMatchType) const
@@ -666,7 +666,7 @@ protected:
 	bool RemoveTagByExplicitName(const FName& TagName);
 
 	/** Adds parent tags for a single tag */
-	void AddParentsForTag(const FGameplayTag& Tag);
+	GAMEPLAYTAGS_API void AddParentsForTag(const FGameplayTag& Tag);
 
 	/** Array of gameplay tags */
 	UPROPERTY(VisibleAnywhere, BlueprintReadWrite, Category=GameplayTags, SaveGame)
@@ -720,9 +720,9 @@ struct TStructOpsTypeTraits<FGameplayTagContainer> : public TStructOpsTypeTraits
 };
 
 /** Class that can be subclassed by a game/plugin to allow easily adding native gameplay tags at startup */
-struct GAMEPLAYTAGS_API FGameplayTagNativeAdder
+struct FGameplayTagNativeAdder
 {
-	FGameplayTagNativeAdder();
+	GAMEPLAYTAGS_API FGameplayTagNativeAdder();
 
 	virtual void AddTags() = 0;
 };
@@ -781,22 +781,22 @@ namespace EGameplayTagQueryStreamVersion
  * Note: these have an extensive details and graph pin customization for editing, so there is no need to expose the internals to Blueprints.
  */
 USTRUCT(BlueprintType, meta=(HasNativeMake="/Script/GameplayTags.BlueprintGameplayTagLibrary.MakeGameplayTagQuery"))
-struct GAMEPLAYTAGS_API FGameplayTagQuery
+struct FGameplayTagQuery
 {
 	GENERATED_BODY();
 
 public:
-	FGameplayTagQuery();
+	GAMEPLAYTAGS_API FGameplayTagQuery();
 
-	FGameplayTagQuery(FGameplayTagQuery const& Other);
-	FGameplayTagQuery(FGameplayTagQuery&& Other);
+	GAMEPLAYTAGS_API FGameplayTagQuery(FGameplayTagQuery const& Other);
+	GAMEPLAYTAGS_API FGameplayTagQuery(FGameplayTagQuery&& Other);
 
 	/** Assignment/Equality operators */
-	FGameplayTagQuery& operator=(FGameplayTagQuery const& Other);
-	FGameplayTagQuery& operator=(FGameplayTagQuery&& Other);
+	GAMEPLAYTAGS_API FGameplayTagQuery& operator=(FGameplayTagQuery const& Other);
+	GAMEPLAYTAGS_API FGameplayTagQuery& operator=(FGameplayTagQuery&& Other);
 
-	bool operator==(const FGameplayTagQuery& Other) const;
-	bool operator!=(const FGameplayTagQuery& Other) const;
+	GAMEPLAYTAGS_API bool operator==(const FGameplayTagQuery& Other) const;
+	GAMEPLAYTAGS_API bool operator!=(const FGameplayTagQuery& Other) const;
 
 private:
 	// Note: Properties need to be editable to allow FComponentPropertyWriter to serialize them, but are hidden in the editor by the customizations mentioned above.
@@ -847,25 +847,25 @@ public:
 	}
 
 	/** Returns true if the given tags match this query, or false otherwise. */
-	bool Matches(FGameplayTagContainer const& Tags) const;
+	GAMEPLAYTAGS_API bool Matches(FGameplayTagContainer const& Tags) const;
 
 	/** Returns true if this query is empty, false otherwise. */
-	bool IsEmpty() const;
+	GAMEPLAYTAGS_API bool IsEmpty() const;
 
 	/** Resets this query to its default empty state. */
-	void Clear();
+	GAMEPLAYTAGS_API void Clear();
 
 	/** Creates this query with the given root expression. */
-	void Build(struct FGameplayTagQueryExpression& RootQueryExpr, FString InUserDescription = FString());
+	GAMEPLAYTAGS_API void Build(struct FGameplayTagQueryExpression& RootQueryExpr, FString InUserDescription = FString());
 
 	/** Static function to assemble and return a query. */
-	static FGameplayTagQuery BuildQuery(struct FGameplayTagQueryExpression& RootQueryExpr, FString InDescription = FString());
+	static GAMEPLAYTAGS_API FGameplayTagQuery BuildQuery(struct FGameplayTagQueryExpression& RootQueryExpr, FString InDescription = FString());
 
 	/** Builds a FGameplayTagQueryExpression from this query. */
-	void GetQueryExpr(struct FGameplayTagQueryExpression& OutExpr) const;
+	GAMEPLAYTAGS_API void GetQueryExpr(struct FGameplayTagQueryExpression& OutExpr) const;
 
 	/** Serialize the tag query */
-	void Serialize(FArchive& Ar);
+	GAMEPLAYTAGS_API void Serialize(FArchive& Ar);
 	
 	/** Returns description string. */
 	const FString& GetDescription() const { return UserDescription.IsEmpty() ? AutoDescription : UserDescription; };
@@ -877,17 +877,17 @@ public:
 	}
 
 	/** Gets the explicit list of all unique gameplay tags referenced by the query. */
-	const TArray<FGameplayTag>& GetGameplayTagArray() const;
+	GAMEPLAYTAGS_API const TArray<FGameplayTag>& GetGameplayTagArray() const;
 
 #if WITH_EDITOR
 	/** Creates this query based on the given EditableQuery object */
-	void BuildFromEditableQuery(class UEditableGameplayTagQuery& EditableQuery); 
+	GAMEPLAYTAGS_API void BuildFromEditableQuery(class UEditableGameplayTagQuery& EditableQuery); 
 
 	/** Creates editable query object tree based on this query */
-	UEditableGameplayTagQuery* CreateEditableQuery() const;
+	GAMEPLAYTAGS_API UEditableGameplayTagQuery* CreateEditableQuery() const;
 #endif // WITH_EDITOR
 
-	static const FGameplayTagQuery EmptyQuery;
+	static GAMEPLAYTAGS_API const FGameplayTagQuery EmptyQuery;
 
 	/**
 	* Shortcuts for easily creating common query types
@@ -895,16 +895,16 @@ public:
 	*/
 
 	/** Creates a tag query that will match if there are any common tags between the given tags and the tags being queries against. */
-	static FGameplayTagQuery MakeQuery_MatchAnyTags(FGameplayTagContainer const& InTags);
-	static FGameplayTagQuery MakeQuery_MatchAllTags(FGameplayTagContainer const& InTags);
-	static FGameplayTagQuery MakeQuery_MatchNoTags(FGameplayTagContainer const& InTags);
+	static GAMEPLAYTAGS_API FGameplayTagQuery MakeQuery_MatchAnyTags(FGameplayTagContainer const& InTags);
+	static GAMEPLAYTAGS_API FGameplayTagQuery MakeQuery_MatchAllTags(FGameplayTagContainer const& InTags);
+	static GAMEPLAYTAGS_API FGameplayTagQuery MakeQuery_MatchNoTags(FGameplayTagContainer const& InTags);
 
-	static FGameplayTagQuery MakeQuery_MatchTag(FGameplayTag const& InTag);
+	static GAMEPLAYTAGS_API FGameplayTagQuery MakeQuery_MatchTag(FGameplayTag const& InTag);
 
 	friend class FQueryEvaluator;
 };
 
-struct GAMEPLAYTAGS_API FGameplayTagQueryExpression
+struct FGameplayTagQueryExpression
 {
 	/** 
 	 * Fluid syntax approach for setting the type of this expression. 
@@ -950,7 +950,7 @@ struct GAMEPLAYTAGS_API FGameplayTagQueryExpression
 	{
 		return AddTag(FName(TagString));
 	}
-	FGameplayTagQueryExpression& AddTag(FName TagName);
+	GAMEPLAYTAGS_API FGameplayTagQueryExpression& AddTag(FName TagName);
 	FGameplayTagQueryExpression& AddTag(FGameplayTag Tag)
 	{
 		ensure(UsesTagSet());
@@ -973,7 +973,7 @@ struct GAMEPLAYTAGS_API FGameplayTagQueryExpression
 	}
 	
 	/** Writes this expression to the given token stream. */
-	void EmitTokens(TArray<uint8>& TokenStream, TArray<FGameplayTag>& TagDictionary) const;
+	GAMEPLAYTAGS_API void EmitTokens(TArray<uint8>& TokenStream, TArray<FGameplayTag>& TagDictionary) const;
 
 	/** Which type of expression this is. */
 	EGameplayTagQueryExprType ExprType;
@@ -996,10 +996,10 @@ struct GAMEPLAYTAGS_API FGameplayTagQueryExpression
 	}
 
 	/** Converts the existing TagQueryExpression into a json object. Returns true on success */
-	bool ConvertToJsonObject(TSharedRef<FJsonObject>& OutObject) const;
+	GAMEPLAYTAGS_API bool ConvertToJsonObject(TSharedRef<FJsonObject>& OutObject) const;
 
 	/** Uses the input json object and fills out the OutQueryExpression with the data. */
-	static bool MakeFromJsonObject(const TSharedRef<FJsonObject>& InObject, FGameplayTagQueryExpression& OutQueryExpression);
+	static GAMEPLAYTAGS_API bool MakeFromJsonObject(const TSharedRef<FJsonObject>& InObject, FGameplayTagQueryExpression& OutQueryExpression);
 };
 
 template<>
@@ -1019,8 +1019,8 @@ struct TStructOpsTypeTraits<FGameplayTagQuery> : public TStructOpsTypeTraitsBase
  * the query struct is rewritten and these UObjects are discarded.
  * This query representation is not intended for runtime use.
  */
-UCLASS(editinlinenew, collapseCategories, Transient) 
-class GAMEPLAYTAGS_API UEditableGameplayTagQuery : public UObject
+UCLASS(editinlinenew, collapseCategories, Transient, MinimalAPI) 
+class UEditableGameplayTagQuery : public UObject
 {
 	GENERATED_BODY()
 
@@ -1038,10 +1038,10 @@ public:
 
 #if WITH_EDITOR
 	/** Converts this editor query construct into the runtime-usable token stream. */
-	void EmitTokens(TArray<uint8>& TokenStream, TArray<FGameplayTag>& TagDictionary, FString* DebugString=nullptr) const;
+	GAMEPLAYTAGS_API void EmitTokens(TArray<uint8>& TokenStream, TArray<FGameplayTag>& TagDictionary, FString* DebugString=nullptr) const;
 
 	/** Generates and returns the export text for this query. */
-	FString GetTagQueryExportText(FGameplayTagQuery const& TagQuery);
+	GAMEPLAYTAGS_API FString GetTagQueryExportText(FGameplayTagQuery const& TagQuery);
 #endif  // WITH_EDITOR
 
 private:
@@ -1050,8 +1050,8 @@ private:
 	FGameplayTagQuery TagQueryExportText_Helper;
 };
 
-UCLASS(abstract, editinlinenew, collapseCategories, Transient)
-class GAMEPLAYTAGS_API UEditableGameplayTagQueryExpression : public UObject
+UCLASS(abstract, editinlinenew, collapseCategories, Transient, MinimalAPI)
+class UEditableGameplayTagQueryExpression : public UObject
 {
 	GENERATED_BODY()
 
@@ -1061,8 +1061,8 @@ public:
 	virtual void EmitTokens(TArray<uint8>& TokenStream, TArray<FGameplayTag>& TagDictionary, FString* DebugString=nullptr) const {};
 
 protected:
-	void EmitTagTokens(FGameplayTagContainer const& TagsToEmit, TArray<uint8>& TokenStream, TArray<FGameplayTag>& TagDictionary, FString* DebugString) const;
-	void EmitExprListTokens(TArray<UEditableGameplayTagQueryExpression*> const& ExprList, TArray<uint8>& TokenStream, TArray<FGameplayTag>& TagDictionary, FString* DebugString) const;
+	GAMEPLAYTAGS_API void EmitTagTokens(FGameplayTagContainer const& TagsToEmit, TArray<uint8>& TokenStream, TArray<FGameplayTag>& TagDictionary, FString* DebugString) const;
+	GAMEPLAYTAGS_API void EmitExprListTokens(TArray<UEditableGameplayTagQueryExpression*> const& ExprList, TArray<uint8>& TokenStream, TArray<FGameplayTag>& TagDictionary, FString* DebugString) const;
 #endif  // WITH_EDITOR
 };
 
diff --git a/Engine/Source/Runtime/GameplayTags/Classes/GameplayTagsManager.h b/Engine/Source/Runtime/GameplayTags/Classes/GameplayTagsManager.h
--- a/Engine/Source/Runtime/GameplayTags/Classes/GameplayTagsManager.h
+++ b/Engine/Source/Runtime/GameplayTags/Classes/GameplayTagsManager.h
@@ -86,7 +86,7 @@ enum class EGameplayTagSelectionType : uint8
 
 /** Struct defining where gameplay tags are loaded/saved from. Mostly for the editor */
 USTRUCT()
-struct GAMEPLAYTAGS_API FGameplayTagSource
+struct FGameplayTagSource
 {
 	GENERATED_USTRUCT_BODY()
 
@@ -117,18 +117,18 @@ struct GAMEPLAYTAGS_API FGameplayTagSource
 	}
 
 	/** Returns the config file that created this source, if valid */
-	FString GetConfigFileName() const;
+	GAMEPLAYTAGS_API FString GetConfigFileName() const;
 
-	static FName GetNativeName();
+	static GAMEPLAYTAGS_API FName GetNativeName();
 
-	static FName GetDefaultName();
+	static GAMEPLAYTAGS_API FName GetDefaultName();
 
 #if WITH_EDITOR
-	static FName GetFavoriteName();
+	static GAMEPLAYTAGS_API FName GetFavoriteName();
 
-	static void SetFavoriteName(FName TagSourceToFavorite);
+	static GAMEPLAYTAGS_API void SetFavoriteName(FName TagSourceToFavorite);
 
-	static FName GetTransientEditorName();
+	static GAMEPLAYTAGS_API FName GetTransientEditorName();
 #endif
 };
 
@@ -298,13 +298,13 @@ private:
 };
 
 /** Holds data about the tag dictionary, is in a singleton UObject */
-UCLASS(config=Engine)
-class GAMEPLAYTAGS_API UGameplayTagsManager : public UObject
+UCLASS(config=Engine, MinimalAPI)
+class UGameplayTagsManager : public UObject
 {
 	GENERATED_UCLASS_BODY()
 
 	/** Destructor */
-	~UGameplayTagsManager();
+	GAMEPLAYTAGS_API ~UGameplayTagsManager();
 
 	/** Returns the global UGameplayTagsManager manager */
 	FORCEINLINE static UGameplayTagsManager& Get()
@@ -328,7 +328,7 @@ class GAMEPLAYTAGS_API UGameplayTagsManager : public UObject
 	* @param ErrorIfNotfound: ensure() that tags exists.
 	*
 	*/
-	void RequestGameplayTagContainer(const TArray<FString>& TagStrings, FGameplayTagContainer& OutTagsContainer, bool bErrorIfNotFound=true) const;
+	GAMEPLAYTAGS_API void RequestGameplayTagContainer(const TArray<FString>& TagStrings, FGameplayTagContainer& OutTagsContainer, bool bErrorIfNotFound=true) const;
 
 	/**
 	 * Gets the FGameplayTag that corresponds to the TagName
@@ -338,7 +338,7 @@ class GAMEPLAYTAGS_API UGameplayTagsManager : public UObject
 	 * 
 	 * @return Will return the corresponding FGameplayTag or an empty one if not found.
 	 */
-	FGameplayTag RequestGameplayTag(FName TagName, bool ErrorIfNotFound=true) const;
+	GAMEPLAYTAGS_API FGameplayTag RequestGameplayTag(FName TagName, bool ErrorIfNotFound=true) const;
 
 	/** 
 	 * Returns true if this is a valid gameplay tag string (foo.bar.baz). If false, it will fill 
@@ -347,13 +347,13 @@ class GAMEPLAYTAGS_API UGameplayTagsManager : public UObject
 	 * @param OutFixedString If non-null and string invalid, will attempt to fix. Will be empty if no fix is possible
 	 * @return True if this can be added to the tag dictionary, false if there's a syntax error
 	 */
-	bool IsValidGameplayTagString(const FString& TagString, FText* OutError = nullptr, FString* OutFixedString = nullptr);
+	GAMEPLAYTAGS_API bool IsValidGameplayTagString(const FString& TagString, FText* OutError = nullptr, FString* OutFixedString = nullptr);
 
 	/**
 	 *	Searches for a gameplay tag given a partial string. This is slow and intended mainly for console commands/utilities to make
 	 *	developer life's easier. This will attempt to match as best as it can. If you pass "A.b" it will match on "A.b." before it matches "a.b.c".
 	 */
-	FGameplayTag FindGameplayTagFromPartialString_Slow(FString PartialString) const;
+	GAMEPLAYTAGS_API FGameplayTag FindGameplayTagFromPartialString_Slow(FString PartialString) const;
 
 	/**
 	 * Registers the given name as a gameplay tag, and tracks that it is being directly referenced from code
@@ -364,21 +364,21 @@ class GAMEPLAYTAGS_API UGameplayTagsManager : public UObject
 	 * 
 	 * @return Will return the corresponding FGameplayTag
 	 */
-	FGameplayTag AddNativeGameplayTag(FName TagName, const FString& TagDevComment = TEXT("(Native)"));
+	GAMEPLAYTAGS_API FGameplayTag AddNativeGameplayTag(FName TagName, const FString& TagDevComment = TEXT("(Native)"));
 
 private:
 	// Only callable from FNativeGameplayTag, these functions do less error checking and can happen after initial tag loading is done
-	void AddNativeGameplayTag(FNativeGameplayTag* TagSource);
-	void RemoveNativeGameplayTag(const FNativeGameplayTag* TagSource);
+	GAMEPLAYTAGS_API void AddNativeGameplayTag(FNativeGameplayTag* TagSource);
+	GAMEPLAYTAGS_API void RemoveNativeGameplayTag(const FNativeGameplayTag* TagSource);
 
 public:
 	/** Call to flush the list of native tags, once called it is unsafe to add more */
-	void DoneAddingNativeTags();
+	GAMEPLAYTAGS_API void DoneAddingNativeTags();
 
-	static FSimpleMulticastDelegate& OnLastChanceToAddNativeTags();
+	static GAMEPLAYTAGS_API FSimpleMulticastDelegate& OnLastChanceToAddNativeTags();
 
 
-	void CallOrRegister_OnDoneAddingNativeTagsDelegate(FSimpleMulticastDelegate::FDelegate Delegate);
+	GAMEPLAYTAGS_API void CallOrRegister_OnDoneAddingNativeTagsDelegate(FSimpleMulticastDelegate::FDelegate Delegate);
 
 	/**
 	 * Gets a Tag Container containing the supplied tag and all of it's parents as explicit tags
@@ -387,7 +387,7 @@ public:
 	 * 
 	 * @return A Tag Container with the supplied tag and all its parents added explicitly
 	 */
-	FGameplayTagContainer RequestGameplayTagParents(const FGameplayTag& GameplayTag) const;
+	GAMEPLAYTAGS_API FGameplayTagContainer RequestGameplayTagParents(const FGameplayTag& GameplayTag) const;
 
 	/**
 	 * Gets a Tag Container containing the all tags in the hierarchy that are children of this tag. Does not return the original tag
@@ -396,10 +396,10 @@ public:
 	 * 
 	 * @return A Tag Container with the supplied tag and all its parents added explicitly
 	 */
-	FGameplayTagContainer RequestGameplayTagChildren(const FGameplayTag& GameplayTag) const;
+	GAMEPLAYTAGS_API FGameplayTagContainer RequestGameplayTagChildren(const FGameplayTag& GameplayTag) const;
 
 	/** Returns direct parent GameplayTag of this GameplayTag, calling on x.y will return x */
-	FGameplayTag RequestGameplayTagDirectParent(const FGameplayTag& GameplayTag) const;
+	GAMEPLAYTAGS_API FGameplayTag RequestGameplayTagDirectParent(const FGameplayTag& GameplayTag) const;
 
 	/**
 	 * Helper function to get the stored TagContainer containing only this tag, which has searchable ParentTags
@@ -482,45 +482,45 @@ public:
 	}
 
 	/** Loads the tag tables referenced in the GameplayTagSettings object */
-	void LoadGameplayTagTables(bool bAllowAsyncLoad = false);
+	GAMEPLAYTAGS_API void LoadGameplayTagTables(bool bAllowAsyncLoad = false);
 
 	/** Loads tag inis contained in the specified path */
-	void AddTagIniSearchPath(const FString& RootDir);
+	GAMEPLAYTAGS_API void AddTagIniSearchPath(const FString& RootDir);
 
 	/** Tries to remove the specified search path, will return true if anything was removed */
-	bool RemoveTagIniSearchPath(const FString& RootDir);
+	GAMEPLAYTAGS_API bool RemoveTagIniSearchPath(const FString& RootDir);
 
 	/** Gets all the current directories to look for tag sources in */
-	void GetTagSourceSearchPaths(TArray<FString>& OutPaths);
+	GAMEPLAYTAGS_API void GetTagSourceSearchPaths(TArray<FString>& OutPaths);
 
 	/** Gets the number of tag source search paths */
-	int32 GetNumTagSourceSearchPaths();
+	GAMEPLAYTAGS_API int32 GetNumTagSourceSearchPaths();
 
 	/** Helper function to construct the gameplay tag tree */
-	void ConstructGameplayTagTree();
+	GAMEPLAYTAGS_API void ConstructGameplayTagTree();
 
 	/** Helper function to destroy the gameplay tag tree */
-	void DestroyGameplayTagTree();
+	GAMEPLAYTAGS_API void DestroyGameplayTagTree();
 
 	/** Splits a tag such as x.y.z into an array of names {x,y,z} */
-	void SplitGameplayTagFName(const FGameplayTag& Tag, TArray<FName>& OutNames) const;
+	GAMEPLAYTAGS_API void SplitGameplayTagFName(const FGameplayTag& Tag, TArray<FName>& OutNames) const;
 
 	/** Gets the list of all tags in the dictionary */
-	void RequestAllGameplayTags(FGameplayTagContainer& TagContainer, bool OnlyIncludeDictionaryTags) const;
+	GAMEPLAYTAGS_API void RequestAllGameplayTags(FGameplayTagContainer& TagContainer, bool OnlyIncludeDictionaryTags) const;
 
 	/** Returns true if if the passed in name is in the tag dictionary and can be created */
-	bool ValidateTagCreation(FName TagName) const;
+	GAMEPLAYTAGS_API bool ValidateTagCreation(FName TagName) const;
 
 	/** Returns the tag source for a given tag source name and type, or null if not found */
-	const FGameplayTagSource* FindTagSource(FName TagSourceName) const;
+	GAMEPLAYTAGS_API const FGameplayTagSource* FindTagSource(FName TagSourceName) const;
 
 	/** Returns the tag source for a given tag source name and type, or null if not found */
-	FGameplayTagSource* FindTagSource(FName TagSourceName);
+	GAMEPLAYTAGS_API FGameplayTagSource* FindTagSource(FName TagSourceName);
 
 	/** Fills in an array with all tag sources of a specific type */
-	void FindTagSourcesWithType(EGameplayTagSourceType TagSourceType, TArray<const FGameplayTagSource*>& OutArray) const;
+	GAMEPLAYTAGS_API void FindTagSourcesWithType(EGameplayTagSourceType TagSourceType, TArray<const FGameplayTagSource*>& OutArray) const;
 
-	void FindTagsWithSource(FStringView PackageNameOrPath, TArray<FGameplayTag>& OutTags) const;
+	GAMEPLAYTAGS_API void FindTagsWithSource(FStringView PackageNameOrPath, TArray<FGameplayTag>& OutTags) const;
 
 	/**
 	 * Check to see how closely two FGameplayTags match. Higher values indicate more matching terms in the tags.
@@ -530,16 +530,16 @@ public:
 	 *
 	 * @return the length of the longest matching pair
 	 */
-	int32 GameplayTagsMatchDepth(const FGameplayTag& GameplayTagOne, const FGameplayTag& GameplayTagTwo) const;
+	GAMEPLAYTAGS_API int32 GameplayTagsMatchDepth(const FGameplayTag& GameplayTagOne, const FGameplayTag& GameplayTagTwo) const;
 
 	/** Returns the number of parents a particular gameplay tag has.  Useful as a quick way to determine which tags may
 	 * be more "specific" than other tags without comparing whether they are in the same hierarchy or anything else.
 	 * Example: "TagA.SubTagA" has 2 Tag Nodes.  "TagA.SubTagA.LeafTagA" has 3 Tag Nodes.
 	 */ 
-	int32 GetNumberOfTagNodes(const FGameplayTag& GameplayTag) const;
+	GAMEPLAYTAGS_API int32 GetNumberOfTagNodes(const FGameplayTag& GameplayTag) const;
 
 	/** Returns true if we should import tags from UGameplayTagsSettings objects (configured by INI files) */
-	bool ShouldImportTagsFromINI() const;
+	GAMEPLAYTAGS_API bool ShouldImportTagsFromINI() const;
 
 	/** Should we print loading errors when trying to load invalid tags */
 	bool ShouldWarnOnInvalidTags() const
@@ -560,38 +560,38 @@ public:
 	}
 
 	/** If we are allowed to unload tags */
-	bool ShouldUnloadTags() const;
+	GAMEPLAYTAGS_API bool ShouldUnloadTags() const;
 
 	/** Returns the hash of NetworkGameplayTagNodeIndex */
 	uint32 GetNetworkGameplayTagNodeIndexHash() const { VerifyNetworkIndex(); return NetworkGameplayTagNodeIndexHash; }
 
 	/** Returns a list of the ini files that contain restricted tags */
-	void GetRestrictedTagConfigFiles(TArray<FString>& RestrictedConfigFiles) const;
+	GAMEPLAYTAGS_API void GetRestrictedTagConfigFiles(TArray<FString>& RestrictedConfigFiles) const;
 
 	/** Returns a list of the source files that contain restricted tags */
-	void GetRestrictedTagSources(TArray<const FGameplayTagSource*>& Sources) const;
+	GAMEPLAYTAGS_API void GetRestrictedTagSources(TArray<const FGameplayTagSource*>& Sources) const;
 
 	/** Returns a list of the owners for a restricted tag config file. May be empty */
-	void GetOwnersForTagSource(const FString& SourceName, TArray<FString>& OutOwners) const;
+	GAMEPLAYTAGS_API void GetOwnersForTagSource(const FString& SourceName, TArray<FString>& OutOwners) const;
 
 	/** Notification that a tag container has been loaded via serialize */
-	void GameplayTagContainerLoaded(FGameplayTagContainer& Container, FProperty* SerializingProperty) const;
+	GAMEPLAYTAGS_API void GameplayTagContainerLoaded(FGameplayTagContainer& Container, FProperty* SerializingProperty) const;
 
 	/** Notification that a gameplay tag has been loaded via serialize */
-	void SingleGameplayTagLoaded(FGameplayTag& Tag, FProperty* SerializingProperty) const;
+	GAMEPLAYTAGS_API void SingleGameplayTagLoaded(FGameplayTag& Tag, FProperty* SerializingProperty) const;
 
 	/** Handles redirectors for an entire container, will also error on invalid tags */
-	void RedirectTagsForContainer(FGameplayTagContainer& Container, FProperty* SerializingProperty) const;
+	GAMEPLAYTAGS_API void RedirectTagsForContainer(FGameplayTagContainer& Container, FProperty* SerializingProperty) const;
 
 	/** Handles redirectors for a single tag, will also error on invalid tag. This is only called for when individual tags are serialized on their own */
-	void RedirectSingleGameplayTag(FGameplayTag& Tag, FProperty* SerializingProperty) const;
+	GAMEPLAYTAGS_API void RedirectSingleGameplayTag(FGameplayTag& Tag, FProperty* SerializingProperty) const;
 
 	/** Handles establishing a single tag from an imported tag name (accounts for redirects too). Called when tags are imported via text. */
-	bool ImportSingleGameplayTag(FGameplayTag& Tag, FName ImportedTagName, bool bImportFromSerialize = false) const;
+	GAMEPLAYTAGS_API bool ImportSingleGameplayTag(FGameplayTag& Tag, FName ImportedTagName, bool bImportFromSerialize = false) const;
 
 	/** Gets a tag name from net index and vice versa, used for replication efficiency */
-	FName GetTagNameFromNetIndex(FGameplayTagNetIndex Index) const;
-	FGameplayTagNetIndex GetNetIndexFromTag(const FGameplayTag &InTag) const;
+	GAMEPLAYTAGS_API FName GetTagNameFromNetIndex(FGameplayTagNetIndex Index) const;
+	GAMEPLAYTAGS_API FGameplayTagNetIndex GetNetIndexFromTag(const FGameplayTag &InTag) const;
 
 	/** Cached number of bits we need to replicate tags. That is, Log2(Number of Tags). Will always be <= 16. */
 	int32 GetNetIndexTrueBitNum() const { VerifyNetworkIndex(); return NetIndexTrueBitNum; }
@@ -610,8 +610,8 @@ public:
 	/** Numbers of bits to use for replicating container size. This can be set via config. */
 	int32 NumBitsForContainerSize;
 
-	void PushDeferOnGameplayTagTreeChangedBroadcast();
-	void PopDeferOnGameplayTagTreeChangedBroadcast();
+	GAMEPLAYTAGS_API void PushDeferOnGameplayTagTreeChangedBroadcast();
+	GAMEPLAYTAGS_API void PopDeferOnGameplayTagTreeChangedBroadcast();
 
 private:
 	/** Cached number of bits we need to replicate tags. That is, Log2(Number of Tags). Will always be <= 16. */
@@ -627,13 +627,13 @@ public:
 
 #if WITH_EDITOR
 	/** Gets a Filtered copy of the GameplayRootTags Array based on the comma delimited filter string passed in */
-	void GetFilteredGameplayRootTags(const FString& InFilterString, TArray< TSharedPtr<FGameplayTagNode> >& OutTagArray) const;
+	GAMEPLAYTAGS_API void GetFilteredGameplayRootTags(const FString& InFilterString, TArray< TSharedPtr<FGameplayTagNode> >& OutTagArray) const;
 
 	/** Returns "Categories" meta property from given handle, used for filtering by tag widget */
-	FString GetCategoriesMetaFromPropertyHandle(TSharedPtr<class IPropertyHandle> PropertyHandle) const;
+	GAMEPLAYTAGS_API FString GetCategoriesMetaFromPropertyHandle(TSharedPtr<class IPropertyHandle> PropertyHandle) const;
 
 	/** Helper function, made to be called by custom OnGetCategoriesMetaFromPropertyHandle handlers  */
-	static FString StaticGetCategoriesMetaFromPropertyHandle(TSharedPtr<class IPropertyHandle> PropertyHandle);
+	static GAMEPLAYTAGS_API FString StaticGetCategoriesMetaFromPropertyHandle(TSharedPtr<class IPropertyHandle> PropertyHandle);
 
 	/** Returns "Categories" meta property from given field, used for filtering by tag widget */
 	template <typename TFieldType>
@@ -648,39 +648,39 @@ public:
 	}
 
 	/** Returns "GameplayTagFilter" meta property from given function, used for filtering by tag widget for any parameters of the function that end up as BP pins */
-	FString GetCategoriesMetaFromFunction(const UFunction* Func, FName ParamName = NAME_None) const;
+	GAMEPLAYTAGS_API FString GetCategoriesMetaFromFunction(const UFunction* Func, FName ParamName = NAME_None) const;
 
 	/** Gets a list of all gameplay tag nodes added by the specific source */
-	void GetAllTagsFromSource(FName TagSource, TArray< TSharedPtr<FGameplayTagNode> >& OutTagArray) const;
+	GAMEPLAYTAGS_API void GetAllTagsFromSource(FName TagSource, TArray< TSharedPtr<FGameplayTagNode> >& OutTagArray) const;
 
 	/** Returns true if this tag is directly in the dictionary already */
-	bool IsDictionaryTag(FName TagName) const;
+	GAMEPLAYTAGS_API bool IsDictionaryTag(FName TagName) const;
 
 	/** Returns information about tag. If not found return false */
-	bool GetTagEditorData(FName TagName, FString& OutComment, FName &OutFirstTagSource, bool& bOutIsTagExplicit, bool &bOutIsRestrictedTag, bool &bOutAllowNonRestrictedChildren) const;
+	GAMEPLAYTAGS_API bool GetTagEditorData(FName TagName, FString& OutComment, FName &OutFirstTagSource, bool& bOutIsTagExplicit, bool &bOutIsRestrictedTag, bool &bOutAllowNonRestrictedChildren) const;
 	
 	/** Returns information about tag. If not found return false */
-    bool GetTagEditorData(FName TagName, FString& OutComment, TArray<FName>& OutTagSources, bool& bOutIsTagExplicit, bool &bOutIsRestrictedTag, bool &bOutAllowNonRestrictedChildren) const;
+    GAMEPLAYTAGS_API bool GetTagEditorData(FName TagName, FString& OutComment, TArray<FName>& OutTagSources, bool& bOutIsTagExplicit, bool &bOutIsRestrictedTag, bool &bOutAllowNonRestrictedChildren) const;
 
 #if WITH_EDITOR
 	/** This is called after EditorRefreshGameplayTagTree. Useful if you need to do anything editor related when tags are added or removed */
-	static FSimpleMulticastDelegate OnEditorRefreshGameplayTagTree;
+	static GAMEPLAYTAGS_API FSimpleMulticastDelegate OnEditorRefreshGameplayTagTree;
 
 	/** Refresh the gameplaytag tree due to an editor change */
-	void EditorRefreshGameplayTagTree();
+	GAMEPLAYTAGS_API void EditorRefreshGameplayTagTree();
 
 	/** Suspends EditorRefreshGameplayTagTree requests */
-	void SuspendEditorRefreshGameplayTagTree(FGuid SuspendToken);
+	GAMEPLAYTAGS_API void SuspendEditorRefreshGameplayTagTree(FGuid SuspendToken);
 
 	/** Resumes EditorRefreshGameplayTagTree requests; triggers a refresh if a request was made while it was suspended */
-	void ResumeEditorRefreshGameplayTagTree(FGuid SuspendToken);
+	GAMEPLAYTAGS_API void ResumeEditorRefreshGameplayTagTree(FGuid SuspendToken);
 #endif //if WITH_EDITOR
 
 	/** Gets a Tag Container containing all of the tags in the hierarchy that are children of this tag, and were explicitly added to the dictionary */
-	FGameplayTagContainer RequestGameplayTagChildrenInDictionary(const FGameplayTag& GameplayTag) const;
+	GAMEPLAYTAGS_API FGameplayTagContainer RequestGameplayTagChildrenInDictionary(const FGameplayTag& GameplayTag) const;
 #if WITH_EDITORONLY_DATA
 	/** Gets a Tag Container containing all of the tags in the hierarchy that are children of this tag, were explicitly added to the dictionary, and do not have any explicitly added tags between them and the specified tag */
-	FGameplayTagContainer RequestGameplayTagDirectDescendantsInDictionary(const FGameplayTag& GameplayTag, EGameplayTagSelectionType SelectionType) const;
+	GAMEPLAYTAGS_API FGameplayTagContainer RequestGameplayTagDirectDescendantsInDictionary(const FGameplayTag& GameplayTag, EGameplayTagSelectionType SelectionType) const;
 #endif // WITH_EDITORONLY_DATA
 
 
@@ -713,20 +713,20 @@ public:
 	DECLARE_MULTICAST_DELEGATE_TwoParams(FOnFilterGameplayTag, const FFilterGameplayTagContext& /** InContext */, bool& /* OUT OutShouldHide */)
 	FOnFilterGameplayTag OnFilterGameplayTag;
 	
-	void NotifyGameplayTagDoubleClickedEditor(FString TagName);
+	GAMEPLAYTAGS_API void NotifyGameplayTagDoubleClickedEditor(FString TagName);
 	
-	bool ShowGameplayTagAsHyperLinkEditor(FString TagName);
+	GAMEPLAYTAGS_API bool ShowGameplayTagAsHyperLinkEditor(FString TagName);
 
 
 #endif //WITH_EDITOR
 
-	void PrintReplicationIndices();
+	GAMEPLAYTAGS_API void PrintReplicationIndices();
 
 #if !(UE_BUILD_SHIPPING || UE_BUILD_TEST)
 	/** Mechanism for tracking what tags are frequently replicated */
 
-	void PrintReplicationFrequencyReport();
-	void NotifyTagReplicated(FGameplayTag Tag, bool WasInContainer);
+	GAMEPLAYTAGS_API void PrintReplicationFrequencyReport();
+	GAMEPLAYTAGS_API void NotifyTagReplicated(FGameplayTag Tag, bool WasInContainer);
 
 	TMap<FGameplayTag, int32>	ReplicationCountMap;
 	TMap<FGameplayTag, int32>	ReplicationCountMap_SingleTags;
@@ -736,13 +736,13 @@ public:
 private:
 
 	/** Initializes the manager */
-	static void InitializeManager();
+	static GAMEPLAYTAGS_API void InitializeManager();
 
 	/** finished loading/adding native tags */
-	static FSimpleMulticastDelegate& OnDoneAddingNativeTagsDelegate();
+	static GAMEPLAYTAGS_API FSimpleMulticastDelegate& OnDoneAddingNativeTagsDelegate();
 
 	/** The Tag Manager singleton */
-	static UGameplayTagsManager* SingletonManager;
+	static GAMEPLAYTAGS_API UGameplayTagsManager* SingletonManager;
 
 	friend class FGameplayTagTest;
 	friend class FGameplayEffectsTest;
@@ -767,18 +767,18 @@ private:
 	 *
 	 * @return Index of the node of the tag
 	 */
-	int32 InsertTagIntoNodeArray(FName Tag, FName FullTag, TSharedPtr<FGameplayTagNode> ParentNode, TArray< TSharedPtr<FGameplayTagNode> >& NodeArray, FName SourceName, const FString& DevComment, bool bIsExplicitTag, bool bIsRestrictedTag, bool bAllowNonRestrictedChildren);
+	GAMEPLAYTAGS_API int32 InsertTagIntoNodeArray(FName Tag, FName FullTag, TSharedPtr<FGameplayTagNode> ParentNode, TArray< TSharedPtr<FGameplayTagNode> >& NodeArray, FName SourceName, const FString& DevComment, bool bIsExplicitTag, bool bIsRestrictedTag, bool bAllowNonRestrictedChildren);
 
 	/** Helper function to populate the tag tree from each table */
-	void PopulateTreeFromDataTable(class UDataTable* Table);
+	GAMEPLAYTAGS_API void PopulateTreeFromDataTable(class UDataTable* Table);
 
-	void AddTagTableRow(const FGameplayTagTableRow& TagRow, FName SourceName, bool bIsRestrictedTag = false);
+	GAMEPLAYTAGS_API void AddTagTableRow(const FGameplayTagTableRow& TagRow, FName SourceName, bool bIsRestrictedTag = false);
 
-	void AddChildrenTags(FGameplayTagContainer& TagContainer, TSharedPtr<FGameplayTagNode> GameplayTagNode, bool RecurseAll=true, bool OnlyIncludeDictionaryTags=false) const;
+	GAMEPLAYTAGS_API void AddChildrenTags(FGameplayTagContainer& TagContainer, TSharedPtr<FGameplayTagNode> GameplayTagNode, bool RecurseAll=true, bool OnlyIncludeDictionaryTags=false) const;
 
-	void AddRestrictedGameplayTagSource(const FString& FileName);
+	GAMEPLAYTAGS_API void AddRestrictedGameplayTagSource(const FString& FileName);
 
-	void AddTagsFromAdditionalLooseIniFiles(const TArray<FString>& IniFileList);
+	GAMEPLAYTAGS_API void AddTagsFromAdditionalLooseIniFiles(const TArray<FString>& IniFileList);
 
 	/**
 	 * Helper function for GameplayTagsMatch to get all parents when doing a parent match,
@@ -787,16 +787,16 @@ private:
 	 * @param NameList		The list we are adding all parent complete names too
 	 * @param GameplayTag	The current Tag we are adding to the list
 	 */
-	void GetAllParentNodeNames(TSet<FName>& NamesList, TSharedPtr<FGameplayTagNode> GameplayTag) const;
+	GAMEPLAYTAGS_API void GetAllParentNodeNames(TSet<FName>& NamesList, TSharedPtr<FGameplayTagNode> GameplayTag) const;
 
 	/** Returns the tag source for a given tag source name, or null if not found */
-	FGameplayTagSource* FindOrAddTagSource(FName TagSourceName, EGameplayTagSourceType SourceType, const FString& RootDirToUse = FString());
+	GAMEPLAYTAGS_API FGameplayTagSource* FindOrAddTagSource(FName TagSourceName, EGameplayTagSourceType SourceType, const FString& RootDirToUse = FString());
 
 	/** Constructs the net indices for each tag */
-	void ConstructNetIndex();
+	GAMEPLAYTAGS_API void ConstructNetIndex();
 
 	/** Marks all of the nodes that descend from CurNode as having an ancestor node that has a source conflict. */
-	void MarkChildrenOfNodeConflict(TSharedPtr<FGameplayTagNode> CurNode);
+	GAMEPLAYTAGS_API void MarkChildrenOfNodeConflict(TSharedPtr<FGameplayTagNode> CurNode);
 
 	void VerifyNetworkIndex() const
 	{
@@ -809,10 +809,10 @@ private:
 	void InvalidateNetworkIndex() { bNetworkIndexInvalidated = true; }
 
 	/** Called in both editor and game when the tag tree changes during startup or editing */
-	void BroadcastOnGameplayTagTreeChanged();
+	GAMEPLAYTAGS_API void BroadcastOnGameplayTagTreeChanged();
 
 	/** Call after modifying the tag tree nodes, this will either call the full editor refresh or a limited game refresh */
-	void HandleGameplayTagTreeChanged(bool bRecreateTree);
+	GAMEPLAYTAGS_API void HandleGameplayTagTreeChanged(bool bRecreateTree);
 
 	// Tag Sources
 	///////////////////////////////////////////////////////
@@ -886,6 +886,6 @@ private:
 	UPROPERTY()
 	TArray<TObjectPtr<UDataTable>> GameplayTagTables;
 
-	const static FName NAME_Categories;
-	const static FName NAME_GameplayTagFilter;
+	GAMEPLAYTAGS_API const static FName NAME_Categories;
+	GAMEPLAYTAGS_API const static FName NAME_GameplayTagFilter;
 };
diff --git a/Engine/Source/Runtime/GameplayTags/Classes/GameplayTagsSettings.h b/Engine/Source/Runtime/GameplayTags/Classes/GameplayTagsSettings.h
--- a/Engine/Source/Runtime/GameplayTags/Classes/GameplayTagsSettings.h
+++ b/Engine/Source/Runtime/GameplayTags/Classes/GameplayTagsSettings.h
@@ -13,7 +13,7 @@
 
 /** Category remapping. This allows base engine tag category meta data to remap to multiple project-specific categories. */
 USTRUCT()
-struct GAMEPLAYTAGS_API FGameplayTagCategoryRemap
+struct FGameplayTagCategoryRemap
 {
 	GENERATED_USTRUCT_BODY()
 
@@ -30,8 +30,8 @@ struct GAMEPLAYTAGS_API FGameplayTagCategoryRemap
 };
 
 /** Base class for storing a list of gameplay tags as an ini list. This is used for both the central list and additional lists */
-UCLASS(config = GameplayTagsList)
-class GAMEPLAYTAGS_API UGameplayTagsList : public UObject
+UCLASS(config = GameplayTagsList, MinimalAPI)
+class UGameplayTagsList : public UObject
 {
 	GENERATED_UCLASS_BODY()
 
@@ -44,12 +44,12 @@ class GAMEPLAYTAGS_API UGameplayTagsList : public UObject
 	TArray<FGameplayTagTableRow> GameplayTagList;
 
 	/** Sorts tags alphabetically */
-	void SortTags();
+	GAMEPLAYTAGS_API void SortTags();
 };
 
 /** Base class for storing a list of restricted gameplay tags as an ini list. This is used for both the central list and additional lists */
-UCLASS(config = GameplayTags)
-class GAMEPLAYTAGS_API URestrictedGameplayTagsList : public UObject
+UCLASS(config = GameplayTags, MinimalAPI)
+class URestrictedGameplayTagsList : public UObject
 {
 	GENERATED_UCLASS_BODY()
 
@@ -62,11 +62,11 @@ class GAMEPLAYTAGS_API URestrictedGameplayTagsList : public UObject
 	TArray<FRestrictedGameplayTagTableRow> RestrictedGameplayTagList;
 
 	/** Sorts tags alphabetically */
-	void SortTags();
+	GAMEPLAYTAGS_API void SortTags();
 };
 
 USTRUCT()
-struct GAMEPLAYTAGS_API FRestrictedConfigInfo
+struct FRestrictedConfigInfo
 {
 	GENERATED_BODY()
 
@@ -77,8 +77,8 @@ struct GAMEPLAYTAGS_API FRestrictedConfigInfo
 	UPROPERTY(config, EditAnywhere, AdvancedDisplay, Category = GameplayTags)
 	TArray<FString> Owners;
 
-	bool operator==(const FRestrictedConfigInfo& Other) const;
-	bool operator!=(const FRestrictedConfigInfo& Other) const;
+	GAMEPLAYTAGS_API bool operator==(const FRestrictedConfigInfo& Other) const;
+	GAMEPLAYTAGS_API bool operator!=(const FRestrictedConfigInfo& Other) const;
 };
 
 /**
@@ -95,8 +95,8 @@ struct GAMEPLAYTAGS_API FRestrictedConfigInfo
  *	the user types (e.g, autocomplete 'Damage.Physical' as the user is adding a 'Damage.Physical.Slash' tag).
  *	
  */
-UCLASS(config=GameplayTags, defaultconfig)
-class GAMEPLAYTAGS_API UGameplayTagsSettings : public UGameplayTagsList
+UCLASS(config=GameplayTags, defaultconfig, MinimalAPI)
+class UGameplayTagsSettings : public UGameplayTagsList
 {
 	GENERATED_UCLASS_BODY()
 
@@ -171,8 +171,8 @@ class GAMEPLAYTAGS_API UGameplayTagsSettings : public UGameplayTagsList
 #endif
 
 #if WITH_EDITOR
-	virtual void PreEditChange(FProperty* PropertyThatWillChange) override;
-	virtual void PostEditChangeProperty(FPropertyChangedEvent& PropertyChangedEvent) override;
+	GAMEPLAYTAGS_API virtual void PreEditChange(FProperty* PropertyThatWillChange) override;
+	GAMEPLAYTAGS_API virtual void PostEditChangeProperty(FPropertyChangedEvent& PropertyChangedEvent) override;
 
 private:
 	// temporary copy of RestrictedConfigFiles that we use to identify changes in the list
@@ -181,12 +181,12 @@ private:
 #endif
 };
 
-UCLASS(config=EditorPerProjectUserSettings, meta=(DisplayName="Gameplay Tag Editing"))
-class GAMEPLAYTAGS_API UGameplayTagsDeveloperSettings : public UDeveloperSettings
+UCLASS(config=EditorPerProjectUserSettings, meta=(DisplayName="Gameplay Tag Editing"), MinimalAPI)
+class UGameplayTagsDeveloperSettings : public UDeveloperSettings
 {
 	GENERATED_UCLASS_BODY()
 
-	virtual FName GetCategoryName() const override;
+	GAMEPLAYTAGS_API virtual FName GetCategoryName() const override;
 
 	/** Allows new tags to be saved into their own INI file. This is make merging easier for non technical developers by setting up their own ini file. */
 	UPROPERTY(config, EditAnywhere, Category=GameplayTags)
diff --git a/Engine/Source/Runtime/GameplayTags/Public/GameplayTagRedirectors.h b/Engine/Source/Runtime/GameplayTags/Public/GameplayTagRedirectors.h
--- a/Engine/Source/Runtime/GameplayTags/Public/GameplayTagRedirectors.h
+++ b/Engine/Source/Runtime/GameplayTags/Public/GameplayTagRedirectors.h
@@ -12,7 +12,7 @@
 
 /** A single redirect from a deleted tag to the new tag that should replace it */
 USTRUCT()
-struct GAMEPLAYTAGS_API FGameplayTagRedirect
+struct FGameplayTagRedirect
 {
 	GENERATED_BODY()
 
@@ -53,4 +53,4 @@ private:
 
 	/** The map of ini-configured tag redirectors */
 	TMap<FName, FGameplayTag> TagRedirects;
-};
\ No newline at end of file
+};
diff --git a/Engine/Source/Runtime/GameplayTags/Public/NativeGameplayTags.h b/Engine/Source/Runtime/GameplayTags/Public/NativeGameplayTags.h
--- a/Engine/Source/Runtime/GameplayTags/Public/NativeGameplayTags.h
+++ b/Engine/Source/Runtime/GameplayTags/Public/NativeGameplayTags.h
@@ -55,14 +55,14 @@ namespace UE::GameplayTags::Private
  * so even if two modules register the same tag and one is unloaded, the tag will still be registered
  * by the other one.
  */
-class GAMEPLAYTAGS_API FNativeGameplayTag : public FNoncopyable
+class FNativeGameplayTag : public FNoncopyable
 {
 public:
-	static FName NAME_NativeGameplayTag;
+	static GAMEPLAYTAGS_API FName NAME_NativeGameplayTag;
 
 public:
-	FNativeGameplayTag(FName PluginName, FName ModuleName, FName TagName, const FString& TagDevComment, ENativeGameplayTagToken);
-	~FNativeGameplayTag();
+	GAMEPLAYTAGS_API FNativeGameplayTag(FName PluginName, FName ModuleName, FName TagName, const FString& TagDevComment, ENativeGameplayTagToken);
+	GAMEPLAYTAGS_API ~FNativeGameplayTag();
 
 	operator FGameplayTag() const { return InternalTag; }
 
@@ -100,14 +100,14 @@ private:
 	FName ModulePackageName;
 	mutable bool bValidated = false;
 
-	void ValidateTagRegistration() const;
+	GAMEPLAYTAGS_API void ValidateTagRegistration() const;
 #endif
 
 #if WITH_EDITORONLY_DATA
 	FString DeveloperComment;
 #endif
 
-	static TSet<const class FNativeGameplayTag*>& GetRegisteredNativeTags();
+	static GAMEPLAYTAGS_API TSet<const class FNativeGameplayTag*>& GetRegisteredNativeTags();
 
 	friend class UGameplayTagsManager;
 };
