Adding thread safety to GameplayTagsManager in game

diff --git a/Engine/Source/Runtime/GameplayTags/Classes/GameplayTagsManager.h b/Engine/Source/Runtime/GameplayTags/Classes/GameplayTagsManager.h
--- a/Engine/Source/Runtime/GameplayTags/Classes/GameplayTagsManager.h
+++ b/Engine/Source/Runtime/GameplayTags/Classes/GameplayTagsManager.h
@@ -604,7 +604,7 @@ public:
 
 	const TArray<TSharedPtr<FGameplayTagNode>>& GetNetworkGameplayTagNodeIndex() const { VerifyNetworkIndex(); return NetworkGameplayTagNodeIndex; }
 
-	DECLARE_MULTICAST_DELEGATE_OneParam(FOnGameplayTagLoaded, const FGameplayTag& /*Tag*/)
+	DECLARE_TS_MULTICAST_DELEGATE_OneParam(FOnGameplayTagLoaded, const FGameplayTag& /*Tag*/)
 	FOnGameplayTagLoaded OnGameplayTagLoadedDelegate;
 
 	/** Numbers of bits to use for replicating container size. This can be set via config. */
@@ -863,11 +863,11 @@ private:
 	/** String with outlawed characters inside tags */
 	FString InvalidTagCharacters;
 
-#if WITH_EDITOR
-	// This critical section is to handle an editor-only issue where tag requests come from another thread when async loading from a background thread in FGameplayTagContainer::Serialize.
+	// This critical section is to handle an issue where tag requests come from another thread when async loading from a background thread in FGameplayTagContainer::Serialize.
 	// This class is not generically threadsafe.
 	mutable FCriticalSection GameplayTagMapCritical;
 
+#if WITH_EDITOR
 	// Transient editor-only tags to support quick-iteration PIE workflows
 	TSet<FName> TransientEditorTags;
 
diff --git a/Engine/Source/Runtime/GameplayTags/Private/GameplayTagsManager.cpp b/Engine/Source/Runtime/GameplayTags/Private/GameplayTagsManager.cpp
--- a/Engine/Source/Runtime/GameplayTags/Private/GameplayTagsManager.cpp
+++ b/Engine/Source/Runtime/GameplayTags/Private/GameplayTagsManager.cpp
@@ -1279,11 +1279,9 @@ int32 UGameplayTagsManager::InsertTagIntoNodeArray(FName Tag, FName FullTag, TSh
 		ensure(GameplayTag.GetTagName() == FullTag);
 
 		{
-#if WITH_EDITOR
-			// This critical section is to handle an editor-only issue where tag requests come from another thread when async loading from a background thread in FGameplayTagContainer::Serialize.
+			// This critical section is to handle an issue where tag requests come from another thread when async loading from a background thread in FGameplayTagContainer::Serialize.
 			// This function is not generically threadsafe.
 			FScopeLock Lock(&GameplayTagMapCritical);
-#endif
 			GameplayTagNodeMap.Add(GameplayTag, TagNode);
 		}
 	}
@@ -1949,11 +1947,9 @@ FGameplayTag UGameplayTagsManager::RequestGameplayTag(FName TagName, bool ErrorI
 {
 	SCOPE_CYCLE_COUNTER(STAT_UGameplayTagsManager_RequestGameplayTag);
 
-#if WITH_EDITOR
-	// This critical section is to handle and editor-only issue where tag requests come from another thread when async loading from a background thread in FGameplayTagContainer::Serialize.
+	// This critical section is to handle an issue where tag requests come from another thread when async loading from a background thread in FGameplayTagContainer::Serialize.
 	// This function is not generically threadsafe.
 	FScopeLock Lock(&GameplayTagMapCritical);
-#endif
 
 	FGameplayTag PossibleTag(TagName);
 
@@ -2045,11 +2041,9 @@ bool UGameplayTagsManager::IsValidGameplayTagString(const FString& TagString, FT
 
 FGameplayTag UGameplayTagsManager::FindGameplayTagFromPartialString_Slow(FString PartialString) const
 {
-#if WITH_EDITOR
-	// This critical section is to handle and editor-only issue where tag requests come from another thread when async loading from a background thread in FGameplayTagContainer::Serialize.
+	// This critical section is to handle an issue where tag requests come from another thread when async loading from a background thread in FGameplayTagContainer::Serialize.
 	// This function is not generically threadsafe.
 	FScopeLock Lock(&GameplayTagMapCritical);
-#endif
 
 	// Exact match first
 	FGameplayTag PossibleTag(*PartialString);
