Making sure FGameplayTag meta properties of Categories and GameplayTagFilter are interchangeable and case insensitive

diff --git a/Engine/Plugins/Editor/GameplayTagsEditor/Source/GameplayTagsEditor/Private/GameplayTagEditorUtilities.cpp b/Engine/Plugins/Editor/GameplayTagsEditor/Source/GameplayTagsEditor/Private/GameplayTagEditorUtilities.cpp
--- a/Engine/Plugins/Editor/GameplayTagsEditor/Source/GameplayTagsEditor/Private/GameplayTagEditorUtilities.cpp
+++ b/Engine/Plugins/Editor/GameplayTagsEditor/Source/GameplayTagsEditor/Private/GameplayTagEditorUtilities.cpp
@@ -11,6 +11,7 @@ namespace UE::GameplayTags::EditorUtilities
 {
 
 static FName NAME_Categories = FName("Categories");
+static FName NAME_GameplayTagFilter = FName("GameplayTagFilter");
 
 FString ExtractTagFilterStringFromGraphPin(UEdGraphPin* InTagPin)
 {
@@ -30,6 +31,10 @@ FString ExtractTagFilterStringFromGraphPin(UEdGraphPin* InTagPin)
 		{
 			FilterString = OwningNode->GetPinMetaData(InTagPin->PinName, NAME_Categories);
 		}
+		if (FilterString.IsEmpty())
+		{
+			FilterString = OwningNode->GetPinMetaData(InTagPin->PinName, NAME_GameplayTagFilter);
+		}
 
 		if (FilterString.IsEmpty())
 		{
diff --git a/Engine/Source/Runtime/GameplayTags/Classes/GameplayTagsManager.h b/Engine/Source/Runtime/GameplayTags/Classes/GameplayTagsManager.h
--- a/Engine/Source/Runtime/GameplayTags/Classes/GameplayTagsManager.h
+++ b/Engine/Source/Runtime/GameplayTags/Classes/GameplayTagsManager.h
@@ -637,18 +637,22 @@ public:
 
 	/** Returns "Categories" meta property from given field, used for filtering by tag widget */
 	template <typename TFieldType>
-	FString GetCategoriesMetaFromField(TFieldType* Field) const
+	static FString GetCategoriesMetaFromField(TFieldType* Field)
 	{
 		check(Field);
 		if (Field->HasMetaData(NAME_Categories))
 		{
 			return Field->GetMetaData(NAME_Categories);
 		}
+		else if (Field->HasMetaData(NAME_GameplayTagFilter))
+		{
+			return Field->GetMetaData(NAME_GameplayTagFilter);
+		}
 		return FString();
 	}
 
 	/** Returns "GameplayTagFilter" meta property from given function, used for filtering by tag widget for any parameters of the function that end up as BP pins */
-	GAMEPLAYTAGS_API FString GetCategoriesMetaFromFunction(const UFunction* Func, FName ParamName = NAME_None) const;
+	static GAMEPLAYTAGS_API FString GetCategoriesMetaFromFunction(const UFunction* Func, FName ParamName = NAME_None);
 
 	/** Gets a list of all gameplay tag nodes added by the specific source */
 	GAMEPLAYTAGS_API void GetAllTagsFromSource(FName TagSource, TArray< TSharedPtr<FGameplayTagNode> >& OutTagArray) const;
diff --git a/Engine/Source/Runtime/GameplayTags/Private/GameplayTagsManager.cpp b/Engine/Source/Runtime/GameplayTags/Private/GameplayTagsManager.cpp
--- a/Engine/Source/Runtime/GameplayTags/Private/GameplayTagsManager.cpp
+++ b/Engine/Source/Runtime/GameplayTags/Private/GameplayTagsManager.cpp
@@ -1449,7 +1449,7 @@ static void RecursiveRootTagSearch(const FString& InFilterString, const TArray<T
 	{
 		FString RootTagName = GameplayRootTags[iTag].Get()->GetSimpleTagName().ToString();
 
-		if (RootTagName.Equals(CurrentFilter) == true)
+		if (RootTagName == CurrentFilter)
 		{
 			if (RestOfFilter.IsEmpty())
 			{
@@ -1531,28 +1531,6 @@ FString UGameplayTagsManager::StaticGetCategoriesMetaFromPropertyHandle(TSharedP
 {
 	FString Categories;
 
-	auto GetFieldMetaData = ([&](FField* Field)
-	{
-		if (Field->HasMetaData(NAME_Categories))
-		{
-			Categories = Field->GetMetaData(NAME_Categories);
-			return true;
-		}
-
-		return false;
-	});
-
-	auto GetMetaData = ([&](UField* Field)
-	{
-		if (Field->HasMetaData(NAME_Categories))
-		{
-			Categories = Field->GetMetaData(NAME_Categories);
-			return true;
-		}
-
-		return false;
-	});
-	
 	while(PropertyHandle.IsValid())
 	{
 		if (FProperty* Property = PropertyHandle->GetProperty())
@@ -1561,7 +1539,8 @@ FString UGameplayTagsManager::StaticGetCategoriesMetaFromPropertyHandle(TSharedP
 			 *	UPROPERTY(EditAnywhere, BlueprintReadWrite, meta = (Categories="GameplayCue"))
 			 *	FGameplayTag GameplayCueTag;
 			 */
-			if (GetFieldMetaData(Property))
+			Categories = GetCategoriesMetaFromField(Property);
+			if (!Categories.IsEmpty())
 			{
 				break;
 			}
@@ -1572,7 +1551,8 @@ FString UGameplayTagsManager::StaticGetCategoriesMetaFromPropertyHandle(TSharedP
 			 */
 			if (FStructProperty* StructProperty = CastField<FStructProperty>(Property))
 			{
-				if (GetMetaData(StructProperty->Struct))
+				Categories = GetCategoriesMetaFromField<UScriptStruct>(StructProperty->Struct);
+				if (!Categories.IsEmpty())
 				{
 					break;
 				}
@@ -1581,7 +1561,8 @@ FString UGameplayTagsManager::StaticGetCategoriesMetaFromPropertyHandle(TSharedP
 			/**	TArray<FGameplayEventKeywordTag> QualifierTagTestList; */
 			if (FArrayProperty* ArrayProperty = CastField<FArrayProperty>(Property))
 			{
-				if (GetFieldMetaData(ArrayProperty->Inner))
+				Categories = GetCategoriesMetaFromField(ArrayProperty->Inner);
+				if (!Categories.IsEmpty())
 				{
 					break;
 				}
@@ -1590,7 +1571,8 @@ FString UGameplayTagsManager::StaticGetCategoriesMetaFromPropertyHandle(TSharedP
 			/**	TMap<FGameplayTag, ValueType> GameplayTagMap; */
 			if (FMapProperty* MapProperty = CastField<FMapProperty>(Property))
 			{
-				if (GetFieldMetaData(MapProperty->KeyProp))
+				Categories = GetCategoriesMetaFromField(MapProperty->KeyProp);
+				if (!Categories.IsEmpty())
 				{
 					break;
 				}
@@ -1602,7 +1584,7 @@ FString UGameplayTagsManager::StaticGetCategoriesMetaFromPropertyHandle(TSharedP
 	return Categories;
 }
 
-FString UGameplayTagsManager::GetCategoriesMetaFromFunction(const UFunction* ThisFunction, FName ParamName /** = NAME_None */) const
+FString UGameplayTagsManager::GetCategoriesMetaFromFunction(const UFunction* ThisFunction, FName ParamName /** = NAME_None */)
 {
 	FString FilterString;
 	if (ThisFunction)
@@ -1613,14 +1595,14 @@ FString UGameplayTagsManager::GetCategoriesMetaFromFunction(const UFunction* Thi
 			FProperty* ParamProp = FindFProperty<FProperty>(ThisFunction, ParamName);
 			if (ParamProp)
 			{
-				FilterString = UGameplayTagsManager::Get().GetCategoriesMetaFromField(ParamProp);
+				FilterString = GetCategoriesMetaFromField(ParamProp);
 			}
 		}
 
 		// No filter found so far, fall back to UFUNCTION-level
-		if (FilterString.IsEmpty() && ThisFunction->HasMetaData(NAME_GameplayTagFilter))
+		if (FilterString.IsEmpty())
 		{
-			FilterString = ThisFunction->GetMetaData(NAME_GameplayTagFilter);
+			FilterString = GetCategoriesMetaFromField(ThisFunction);
 		}
 	}
 
