[NO NON-PII COMMIT MESSAGE]

diff --git a/Engine/Source/Runtime/GameplayTags/Private/GameplayTagsManager.cpp b/Engine/Source/Runtime/GameplayTags/Private/GameplayTagsManager.cpp
--- a/Engine/Source/Runtime/GameplayTags/Private/GameplayTagsManager.cpp
+++ b/Engine/Source/Runtime/GameplayTags/Private/GameplayTagsManager.cpp
@@ -1933,13 +1933,42 @@ FGameplayTag UGameplayTagsManager::RequestGameplayTag(FName TagName, bool ErrorI
 	// This function is not generically threadsafe.
 	FScopeLock Lock(&GameplayTagMapCritical);
 
-	FGameplayTag PossibleTag(TagName);
+	// Check if there are redirects for this tag. If so and the redirected tag is in the node map, return it.
+	// Redirects take priority, even if the tag itself may exist.
+	if (const FGameplayTag* RedirectedTag = FGameplayTagRedirectors::Get().RedirectTag(TagName))
+	{
+		// Check if the redirected tag exists in the node map
+		if (GameplayTagNodeMap.Contains(*RedirectedTag))
+		{
+			return *RedirectedTag;
+		}
+
+		// The tag that was redirected to was not found. Error if that was requested.
+		if (ErrorIfNotFound)
+		{
+			static TSet<FName> MissingRedirectedTagNames;
+			if (!MissingRedirectedTagNames.Contains(TagName))
+			{
+				const FString RedirectedToName = RedirectedTag->GetTagName().ToString();
+				ensureAlwaysMsgf(false, TEXT("Requested Gameplay Tag %s was redirected to %s but %s was not found. Fix or remove the redirect from config."), *TagName.ToString(), *RedirectedToName, *RedirectedToName);
+				MissingRedirectedTagNames.Add(TagName);
+			}
+		}
+		
+		// TagName got redirected to a non-existent tag. We'll return an empty tag rather than falling through
+		// and trying to resolve the original tag name. Stale redirects should be fixed.
+		return FGameplayTag();
+	}
 
+	// Check if the tag itself exists in the node map. If so, return it.
+	const FGameplayTag PossibleTag(TagName);
 	if (GameplayTagNodeMap.Contains(PossibleTag))
 	{
 		return PossibleTag;
 	}
-	else if (ErrorIfNotFound)
+
+	// The tag is not found. Error if that was requested.
+	if (ErrorIfNotFound)
 	{
 		static TSet<FName> MissingTagName;
 		if (!MissingTagName.Contains(TagName))
