Fix it so gameplay tag containers that are serialized before tags are registered get correct parent tags instead of silently corrupting them.
Fix thread safety issues by adding critical sections, changing GetSingleTagContainer to return a value and deprecating the version returning a raw pointer
Significantly optimize RequestAllGameplayTags and cleanup comments.
Add FGameplayTag::ParseParentTags, reorganize tag manager accesses, and add a few more tag tests.

diff --git a/Engine/Source/Runtime/GameplayTags/Classes/GameplayTagContainer.h b/Engine/Source/Runtime/GameplayTags/Classes/GameplayTagContainer.h
--- a/Engine/Source/Runtime/GameplayTags/Classes/GameplayTagContainer.h
+++ b/Engine/Source/Runtime/GameplayTags/Classes/GameplayTagContainer.h
@@ -150,15 +150,24 @@ struct FGameplayTag
 		return (TagName != NAME_None);
 	}
 
-	/** Returns reference to a GameplayTagContainer containing only this tag */
-	GAMEPLAYTAGS_API const FGameplayTagContainer& GetSingleTagContainer() const;
+	/** Returns a GameplayTagContainer containing only this tag */
+	GAMEPLAYTAGS_API FGameplayTagContainer GetSingleTagContainer() const;
 
 	/** Returns direct parent GameplayTag of this GameplayTag, calling on x.y will return x */
 	GAMEPLAYTAGS_API FGameplayTag RequestDirectParent() const;
 
-	/** Returns a new container explicitly containing the tags of this tag */
+	/** 
+	 * Returns a new tag container that includes this tag and all parent tags as explicitly added tags. 
+	 * For example, calling this on x.y.z would return a tag container with x.y.z, x.y, and x
+	 */
 	GAMEPLAYTAGS_API FGameplayTagContainer GetGameplayTagParents() const;
 
+	/** 
+	 * Parses the tag name and fills in UniqueParentTags with raw parent tags, without validating with the tag manager.
+	 * For example, calling this on x.y.z would add x.y and x to UniqueParentTags if they were not already in the array
+	 */
+	GAMEPLAYTAGS_API void ParseParentTags(TArray<FGameplayTag>& UniqueParentTags) const;
+
 	/** Used so we can have a TMap of this struct */
 	FORCEINLINE friend uint32 GetTypeHash(const FGameplayTag& Tag)
 	{
@@ -665,7 +674,7 @@ protected:
 	 */
 	bool RemoveTagByExplicitName(const FName& TagName);
 
-	/** Adds parent tags for a single tag */
+	UE_DEPRECATED(5.4, "Use ParseParentTags or ExtractParentTags instead")
 	GAMEPLAYTAGS_API void AddParentsForTag(const FGameplayTag& Tag);
 
 	/** Array of gameplay tags */
diff --git a/Engine/Source/Runtime/GameplayTags/Classes/GameplayTagsManager.h b/Engine/Source/Runtime/GameplayTags/Classes/GameplayTagsManager.h
--- a/Engine/Source/Runtime/GameplayTags/Classes/GameplayTagsManager.h
+++ b/Engine/Source/Runtime/GameplayTags/Classes/GameplayTagsManager.h
@@ -381,14 +381,28 @@ public:
 	GAMEPLAYTAGS_API void CallOrRegister_OnDoneAddingNativeTagsDelegate(FSimpleMulticastDelegate::FDelegate Delegate);
 
 	/**
-	 * Gets a Tag Container containing the supplied tag and all of it's parents as explicit tags
+	 * Gets a Tag Container containing the supplied tag and all of its parents as explicit tags.
+	 * For example, passing in x.y.z would return a tag container with x.y.z, x.y, and x.
+	 * This will only work for tags that have been properly registered.
 	 *
-	 * @param GameplayTag The Tag to use at the child most tag for this container
+	 * @param GameplayTag The tag to use at the child most tag for this container
 	 * 
-	 * @return A Tag Container with the supplied tag and all its parents added explicitly
+	 * @return A tag container with the supplied tag and all its parents added explicitly, or an empty container if that failed
 	 */
 	GAMEPLAYTAGS_API FGameplayTagContainer RequestGameplayTagParents(const FGameplayTag& GameplayTag) const;
 
+	/**
+	 * Fills in an array of gameplay tags with all of tags that are the parents of the passed in tag.
+	 * For example, passing in x.y.z would add x.y and x to UniqueParentTags if they was not already there.
+	 * This is used by the GameplayTagContainer code and may work for unregistered tags depending on serialization settings.
+	 *
+	 * @param GameplayTag The gameplay tag to extract parent tags from
+	 * @param UniqueParentTags A list of parent tags that will be added to if necessary
+	 *
+	 * @return true if any tags were added to UniqueParentTags
+	 */
+	GAMEPLAYTAGS_API bool ExtractParentTags(const FGameplayTag& GameplayTag, TArray<FGameplayTag>& UniqueParentTags) const;
+
 	/**
 	 * Gets a Tag Container containing the all tags in the hierarchy that are children of this tag. Does not return the original tag
 	 *
@@ -401,37 +415,10 @@ public:
 	/** Returns direct parent GameplayTag of this GameplayTag, calling on x.y will return x */
 	GAMEPLAYTAGS_API FGameplayTag RequestGameplayTagDirectParent(const FGameplayTag& GameplayTag) const;
 
-	/**
-	 * Helper function to get the stored TagContainer containing only this tag, which has searchable ParentTags
-	 * @param GameplayTag		Tag to get single container of
-	 * @return					Pointer to container with this tag
-	 */
+	UE_DEPRECATED(5.4, "This function is not threadsafe, use FindTagNode or FGameplayTag::GetSingleTagContainer")
 	FORCEINLINE_DEBUGGABLE const FGameplayTagContainer* GetSingleTagContainer(const FGameplayTag& GameplayTag) const
 	{
-		// Doing this with pointers to avoid a shared ptr reference count change
-		const TSharedPtr<FGameplayTagNode>* Node = GameplayTagNodeMap.Find(GameplayTag);
-
-		if (Node)
-		{
-			return &(*Node)->GetSingleTagContainer();
-		}
-#if WITH_EDITOR
-		// Check redirector
-		if (GIsEditor && GameplayTag.IsValid())
-		{
-			FGameplayTag RedirectedTag = GameplayTag;
-
-			RedirectSingleGameplayTag(RedirectedTag, nullptr);
-
-			Node = GameplayTagNodeMap.Find(RedirectedTag);
-
-			if (Node)
-			{
-				return &(*Node)->GetSingleTagContainer();
-			}
-		}
-#endif
-		return nullptr;
+		return GetSingleTagContainerPtr(GameplayTag);
 	}
 
 	/**
@@ -443,6 +430,8 @@ public:
 	 */
 	FORCEINLINE_DEBUGGABLE TSharedPtr<FGameplayTagNode> FindTagNode(const FGameplayTag& GameplayTag) const
 	{
+		FScopeLock Lock(&GameplayTagMapCritical);
+
 		const TSharedPtr<FGameplayTagNode>* Node = GameplayTagNodeMap.Find(GameplayTag);
 
 		if (Node)
@@ -505,7 +494,7 @@ public:
 	/** Splits a tag such as x.y.z into an array of names {x,y,z} */
 	GAMEPLAYTAGS_API void SplitGameplayTagFName(const FGameplayTag& Tag, TArray<FName>& OutNames) const;
 
-	/** Gets the list of all tags in the dictionary */
+	/** Gets the list of all registered tags, setting OnlyIncludeDictionaryTags will exclude implicitly added tags if possible */
 	GAMEPLAYTAGS_API void RequestAllGameplayTags(FGameplayTagContainer& TagContainer, bool OnlyIncludeDictionaryTags) const;
 
 	/** Returns true if if the passed in name is in the tag dictionary and can be created */
@@ -657,7 +646,7 @@ public:
 	/** Gets a list of all gameplay tag nodes added by the specific source */
 	GAMEPLAYTAGS_API void GetAllTagsFromSource(FName TagSource, TArray< TSharedPtr<FGameplayTagNode> >& OutTagArray) const;
 
-	/** Returns true if this tag is directly in the dictionary already */
+	/** Returns true if this tag was explicitly registered, this is false for implictly added parent tags */
 	GAMEPLAYTAGS_API bool IsDictionaryTag(FName TagName) const;
 
 	/** Returns information about tag. If not found return false */
@@ -756,6 +745,41 @@ private:
 	friend class SAddNewGameplayTagSourceWidget;
 	friend class FNativeGameplayTag;
 
+	/**
+	 * Helper function to get the stored TagContainer containing only this tag, which has searchable ParentTags
+	 * NOTE: This function is not threadsafe and should only be used in code that locks the tag map critical section
+	 * @param GameplayTag		Tag to get single container of
+	 * @return					Pointer to container with this tag
+	 */
+	FORCEINLINE_DEBUGGABLE const FGameplayTagContainer* GetSingleTagContainerPtr(const FGameplayTag& GameplayTag) const
+	{
+		// Doing this with pointers to avoid a shared ptr reference count change
+		const TSharedPtr<FGameplayTagNode>* Node = GameplayTagNodeMap.Find(GameplayTag);
+
+		if (Node)
+		{
+			return &(*Node)->GetSingleTagContainer();
+		}
+#if WITH_EDITOR
+		// Check redirector
+		if (GIsEditor && GameplayTag.IsValid())
+		{
+			FGameplayTag RedirectedTag = GameplayTag;
+
+			RedirectSingleGameplayTag(RedirectedTag, nullptr);
+
+			Node = GameplayTagNodeMap.Find(RedirectedTag);
+
+			if (Node)
+			{
+				return &(*Node)->GetSingleTagContainer();
+			}
+		}
+#endif
+		return nullptr;
+	}
+
+
 	/**
 	 * Helper function to insert a tag into a tag node array
 	 *
diff --git a/Engine/Source/Runtime/GameplayTags/Private/GameplayTagContainer.cpp b/Engine/Source/Runtime/GameplayTags/Private/GameplayTagContainer.cpp
--- a/Engine/Source/Runtime/GameplayTags/Private/GameplayTagContainer.cpp
+++ b/Engine/Source/Runtime/GameplayTags/Private/GameplayTagContainer.cpp
@@ -564,12 +564,7 @@ bool FGameplayTagContainer::ComplexHasTag(FGameplayTag const& TagToCheck, TEnumA
 	}
 	else
 	{
-		const FGameplayTagContainer* SingleContainer = UGameplayTagsManager::Get().GetSingleTagContainer(TagToCheck);
-		if (SingleContainer && SingleContainer->DoesTagContainerMatch(*this, EGameplayTagMatchType::IncludeParentTags, EGameplayTagMatchType::Explicit, EGameplayContainerMatchType::Any))
-		{
-			return true;
-		}
-
+		return TagToCheck.GetSingleTagContainer().DoesTagContainerMatch(*this, EGameplayTagMatchType::IncludeParentTags, EGameplayTagMatchType::Explicit, EGameplayContainerMatchType::Any);
 	}
 	return false;
 }
@@ -593,27 +588,24 @@ bool FGameplayTagContainer::RemoveTagByExplicitName(const FName& TagName)
 	return false;
 }
 
-FORCEINLINE_DEBUGGABLE void FGameplayTagContainer::AddParentsForTag(const FGameplayTag& Tag)
+void FGameplayTagContainer::AddParentsForTag(const FGameplayTag& Tag)
 {
-	const FGameplayTagContainer* SingleContainer = UGameplayTagsManager::Get().GetSingleTagContainer(Tag);
-	if (SingleContainer)
-	{
-		// Add Parent tags from this tag to our own
-		for (const FGameplayTag& ParentTag : SingleContainer->ParentTags)
-		{
-			ParentTags.AddUnique(ParentTag);
-		}
-	}
+	UGameplayTagsManager::Get().ExtractParentTags(Tag, ParentTags);
 }
+
 void FGameplayTagContainer::FillParentTags()
 {
 	SCOPE_CYCLE_COUNTER(STAT_FGameplayTagContainer_FillParentTags);
 
 	ParentTags.Reset();
 
-	for (const FGameplayTag& Tag : GameplayTags)
+	if (GameplayTags.Num() > 0)
 	{
-		AddParentsForTag(Tag);
+		UGameplayTagsManager& TagManager = UGameplayTagsManager::Get();
+		for (const FGameplayTag& Tag : GameplayTags)
+		{
+			TagManager.ExtractParentTags(Tag, ParentTags);
+		}
 	}
 }
 
@@ -704,9 +696,6 @@ void FGameplayTagContainer::AppendMatchingTags(FGameplayTagContainer const& Othe
 	}
 }
 
-static UGameplayTagsManager* CachedTagManager = nullptr;
-
-
 void FGameplayTagContainer::AddTag(const FGameplayTag& TagToAdd)
 {
 	CONDITIONAL_SCOPE_CYCLE_COUNTER(STAT_FGameplayTagContainer_AddTag, GEnableGameplayTagDetailedStats);
@@ -716,20 +705,18 @@ void FGameplayTagContainer::AddTag(const FGameplayTag& TagToAdd)
 		// Don't want duplicate tags
 		GameplayTags.AddUnique(TagToAdd);
 
-		AddParentsForTag(TagToAdd);
+		UGameplayTagsManager::Get().ExtractParentTags(TagToAdd, ParentTags);
 	}
 }
 
 void FGameplayTagContainer::AddTagFast(const FGameplayTag& TagToAdd)
 {
 	GameplayTags.Add(TagToAdd);
-	AddParentsForTag(TagToAdd);
+	UGameplayTagsManager::Get().ExtractParentTags(TagToAdd, ParentTags);
 }
 
 bool FGameplayTagContainer::AddLeafTag(const FGameplayTag& TagToAdd)
 {
-	UGameplayTagsManager& TagManager = UGameplayTagsManager::Get();
-
 	// Check tag is not already explicitly in container
 	if (HasTagExact(TagToAdd))
 	{
@@ -742,20 +729,17 @@ bool FGameplayTagContainer::AddLeafTag(const FGameplayTag& TagToAdd)
 		return false;
 	}
 
-	const FGameplayTagContainer* TagToAddContainer = UGameplayTagsManager::Get().GetSingleTagContainer(TagToAdd);
-
-	// This should always succeed
-	if (!ensure(TagToAddContainer))
-	{
-		return false;
-	}
+	TSharedPtr<FGameplayTagNode> TagNode = UGameplayTagsManager::Get().FindTagNode(TagToAdd);
 
-	// Remove any tags in the container that are a parent to TagToAdd
-	for (const FGameplayTag& ParentTag : TagToAddContainer->ParentTags)
+	if (ensureMsgf(TagNode.IsValid(), TEXT("AddLeafTag passed invalid gameplay tag %s, only registered tags can be queried"), *TagToAdd.GetTagName().ToString()))
 	{
-		if (HasTagExact(ParentTag))
+		// Remove any tags in the container that are a parent to TagToAdd
+		for (const FGameplayTag& ParentTag : TagNode->GetSingleTagContainer().ParentTags)
 		{
-			RemoveTag(ParentTag);
+			if (HasTagExact(ParentTag))
+			{
+				RemoveTag(ParentTag);
+			}
 		}
 	}
 
@@ -1041,19 +1025,19 @@ FText FGameplayTagContainer::ToMatchingText(EGameplayContainerMatchType MatchTyp
 	return FText::Format(MatchingDescription[DescriptionIndex], Arguments);
 }
 
-const FGameplayTagContainer& FGameplayTag::GetSingleTagContainer() const
+FGameplayTagContainer FGameplayTag::GetSingleTagContainer() const
 {
 	SCOPE_CYCLE_COUNTER(STAT_FGameplayTag_GetSingleTagContainer);
 
-	const FGameplayTagContainer* TagContainer = UGameplayTagsManager::Get().GetSingleTagContainer(*this);
+	TSharedPtr<FGameplayTagNode> TagNode = UGameplayTagsManager::Get().FindTagNode(*this);
 
-	if (TagContainer)
+	if (TagNode.IsValid())
 	{
-		return *TagContainer;
+		return TagNode->GetSingleTagContainer();
 	}
 
-	// This should always be invalid if the node is missing
-	ensure(!IsValid());
+	// This tag should always be invalid if the node is missing
+	ensureMsgf(!IsValid(), TEXT("GetSingleTagContainer passed invalid gameplay tag %s, only registered tags can be queried"), *GetTagName().ToString());
 
 	return FGameplayTagContainer::EmptyContainer;
 }
@@ -1073,15 +1057,37 @@ FGameplayTagContainer FGameplayTag::GetGameplayTagParents() const
 	return UGameplayTagsManager::Get().RequestGameplayTagParents(*this);
 }
 
+void FGameplayTag::ParseParentTags(TArray<FGameplayTag>& UniqueParentTags) const
+{
+	// This needs to be in the same order as the gameplay tag node ParentTags, which is immediate parent first
+	FName RawTag = GetTagName();
+	TStringBuilder<FName::StringBufferSize> TagBuffer(InPlace, RawTag);
+	FStringView TagView = TagBuffer.ToView();
+
+	int32 DotIndex = UE::String::FindLastChar(TagView, TEXT('.'));
+
+	while (DotIndex != INDEX_NONE)
+	{
+		// Remove everything starting with the last dot
+		TagView.LeftInline(DotIndex);
+		DotIndex = UE::String::FindLastChar(TagView, TEXT('.'));
+
+		// Add the name to the array
+		FGameplayTag ParentTag = FGameplayTag(FName(TagView));
+
+		UniqueParentTags.AddUnique(MoveTemp(ParentTag));
+	}
+}
+
 bool FGameplayTag::MatchesTag(const FGameplayTag& TagToCheck) const
 {
 	CONDITIONAL_SCOPE_CYCLE_COUNTER(STAT_FGameplayTag_MatchesTag, GEnableGameplayTagDetailedStats);
 
-	const FGameplayTagContainer* TagContainer = UGameplayTagsManager::Get().GetSingleTagContainer(*this);
+	TSharedPtr<FGameplayTagNode> TagNode = UGameplayTagsManager::Get().FindTagNode(*this);
 
-	if (TagContainer)
+	if (TagNode.IsValid())
 	{
-		return TagContainer->HasTag(TagToCheck);
+		return TagNode->GetSingleTagContainer().HasTag(TagToCheck);
 	}
 
 	// If a non-empty tag has not been registered, it will not exist in the tag database so this function may return the incorrect value
@@ -1095,11 +1101,11 @@ bool FGameplayTag::MatchesAny(const FGameplayTagContainer& ContainerToCheck) con
 {
 	CONDITIONAL_SCOPE_CYCLE_COUNTER(STAT_FGameplayTag_MatchesAny, GEnableGameplayTagDetailedStats);
 
-	const FGameplayTagContainer* TagContainer = UGameplayTagsManager::Get().GetSingleTagContainer(*this);
+	TSharedPtr<FGameplayTagNode> TagNode = UGameplayTagsManager::Get().FindTagNode(*this);
 
-	if (TagContainer)
+	if (TagNode.IsValid())
 	{
-		return TagContainer->HasAny(ContainerToCheck);
+		return TagNode->GetSingleTagContainer().HasAny(ContainerToCheck);
 	}
 
 	// If a non-empty tag has not been registered, it will not exist in the tag database so this function may return the incorrect value
@@ -1229,7 +1235,7 @@ bool FGameplayTag::NetSerialize_Packed(FArchive& Ar, class UPackageMap* Map, boo
 					TagName = NetFieldExportGroup->NetFieldExports[NetIndex].ExportName;
 
 					// Validate the tag name
-					const FGameplayTag Tag = UGameplayTagsManager::Get().RequestGameplayTag(TagName, false);
+					const FGameplayTag Tag = TagManager.RequestGameplayTag(TagName, false);
 
 					// Warn (once) if the tag isn't found
 					if (!Tag.IsValid() && !NetFieldExportGroup->NetFieldExports[NetIndex].bIncompatible)
diff --git a/Engine/Source/Runtime/GameplayTags/Private/GameplayTagsManager.cpp b/Engine/Source/Runtime/GameplayTags/Private/GameplayTagsManager.cpp
--- a/Engine/Source/Runtime/GameplayTags/Private/GameplayTagsManager.cpp
+++ b/Engine/Source/Runtime/GameplayTags/Private/GameplayTagsManager.cpp
@@ -588,6 +588,8 @@ int32 PrintNetIndiceAssignment = 0;
 static FAutoConsoleVariableRef CVarPrintNetIndiceAssignment(TEXT("GameplayTags.PrintNetIndiceAssignment"), PrintNetIndiceAssignment, TEXT("Logs GameplayTag NetIndice assignment"), ECVF_Default );
 void UGameplayTagsManager::ConstructNetIndex()
 {
+	FScopeLock Lock(&GameplayTagMapCritical);
+
 	bNetworkIndexInvalidated = false;
 
 	NetworkGameplayTagNodeIndex.Empty();
@@ -1197,6 +1199,8 @@ UGameplayTagsManager::~UGameplayTagsManager()
 
 void UGameplayTagsManager::DestroyGameplayTagTree()
 {
+	FScopeLock Lock(&GameplayTagMapCritical);
+
 	if (GameplayRootTag.IsValid())
 	{
 		GameplayRootTag->ResetNode();
@@ -1305,6 +1309,8 @@ void UGameplayTagsManager::PrintReplicationIndices()
 
 	UE_LOG(LogGameplayTags, Display, TEXT("::PrintReplicationIndices (TOTAL %d)"), GameplayTagNodeMap.Num());
 
+	FScopeLock Lock(&GameplayTagMapCritical);
+
 	for (auto It : GameplayTagNodeMap)
 	{
 		FGameplayTag Tag = It.Key;
@@ -1611,6 +1617,8 @@ FString UGameplayTagsManager::GetCategoriesMetaFromFunction(const UFunction* Thi
 
 void UGameplayTagsManager::GetAllTagsFromSource(FName TagSource, TArray< TSharedPtr<FGameplayTagNode> >& OutTagArray) const
 {
+	FScopeLock Lock(&GameplayTagMapCritical);
+
 	for (const TPair<FGameplayTag, TSharedPtr<FGameplayTagNode>>& NodePair : GameplayTagNodeMap)
 	{
 		if (NodePair.Value->SourceNames.Contains(TagSource))
@@ -2183,7 +2191,9 @@ void UGameplayTagsManager::DoneAddingNativeTags()
 
 FGameplayTagContainer UGameplayTagsManager::RequestGameplayTagParents(const FGameplayTag& GameplayTag) const
 {
-	const FGameplayTagContainer* ParentTags = GetSingleTagContainer(GameplayTag);
+	FScopeLock Lock(&GameplayTagMapCritical);
+
+	const FGameplayTagContainer* ParentTags = GetSingleTagContainerPtr(GameplayTag);
 
 	if (ParentTags)
 	{
@@ -2192,22 +2202,62 @@ FGameplayTagContainer UGameplayTagsManager::RequestGameplayTagParents(const FGam
 	return FGameplayTagContainer();
 }
 
+// If true, verify that the node lookup and manual methods give identical results
+#define VALIDATE_EXTRACT_PARENT_TAGS 0
+
+bool UGameplayTagsManager::ExtractParentTags(const FGameplayTag& GameplayTag, TArray<FGameplayTag>& UniqueParentTags) const
+{
+	// This gets called during GameplayTagContainer serialization so needs to be efficient
+	if (!GameplayTag.IsValid())
+	{
+		return false;
+	}
+
+#if VALIDATE_EXTRACT_PARENT_TAGS
+	TArray<FGameplayTag> ValidationCopy = UniqueParentTags;
+#endif
+
+	FScopeLock Lock(&GameplayTagMapCritical);
+
+	int32 OldSize = UniqueParentTags.Num();
+	FName RawTag = GameplayTag.GetTagName();
+
+	// This code does not check redirectors because that was already handled by GameplayTagContainerLoaded
+	const TSharedPtr<FGameplayTagNode>* Node = GameplayTagNodeMap.Find(GameplayTag);
+	if (Node)
+	{
+		// Use the registered tag container if it exists
+		const FGameplayTagContainer& SingleContainer = (*Node)->GetSingleTagContainer();
+		for (const FGameplayTag& ParentTag : SingleContainer.ParentTags)
+		{
+			UniqueParentTags.AddUnique(ParentTag);
+		}
+
+#if VALIDATE_EXTRACT_PARENT_TAGS
+		GameplayTag.ParseParentTags(ValidationCopy);
+
+		ensureAlwaysMsgf(ValidationCopy == UniqueParentTags, TEXT("ExtractParentTags results are inconsistent for tag %s"), *GameplayTag.ToString());
+#endif
+	}
+	else if (!ShouldClearInvalidTags())
+	{
+		// If we don't clear invalid tags, we need to extract the parents now in case they get registered later
+		GameplayTag.ParseParentTags(UniqueParentTags);
+	}
+
+	return UniqueParentTags.Num() != OldSize;
+}
+
 void UGameplayTagsManager::RequestAllGameplayTags(FGameplayTagContainer& TagContainer, bool OnlyIncludeDictionaryTags) const
 {
-	TArray<TSharedPtr<FGameplayTagNode>> ValueArray;
-	GameplayTagNodeMap.GenerateValueArray(ValueArray);
-	for (const TSharedPtr<FGameplayTagNode>& TagNode : ValueArray)
+	FScopeLock Lock(&GameplayTagMapCritical);
+
+	for (const TPair<FGameplayTag, TSharedPtr<FGameplayTagNode>>& NodePair : GameplayTagNodeMap)
 	{
-#if WITH_EDITOR
-		bool DictTag = IsDictionaryTag(TagNode->GetCompleteTagName());
-#else
-		bool DictTag = false;
-#endif 
-		if (!OnlyIncludeDictionaryTags || DictTag)
+		const TSharedPtr<FGameplayTagNode>& TagNode = NodePair.Value;
+		if (!OnlyIncludeDictionaryTags || TagNode->IsExplicitTag())
 		{
-			const FGameplayTag* Tag = GameplayTagNodeMap.FindKey(TagNode);
-			check(Tag);
-			TagContainer.AddTagFast(*Tag);
+			TagContainer.AddTagFast(TagNode->GetCompleteTag());
 		}
 	}
 }
diff --git a/Engine/Source/Runtime/GameplayTags/Private/Tests/GameplayTagTests.cpp b/Engine/Source/Runtime/GameplayTags/Private/Tests/GameplayTagTests.cpp
--- a/Engine/Source/Runtime/GameplayTags/Private/Tests/GameplayTagTests.cpp
+++ b/Engine/Source/Runtime/GameplayTags/Private/Tests/GameplayTagTests.cpp
@@ -221,6 +221,18 @@ public:
 	
 		TestTrueExpr(FilteredTagContainer.HasTagExact(CueTag));
 		TestTrueExpr(!FilteredTagContainer.HasTagExact(EffectDamage1Tag));
+
+		FGameplayTagContainer SingleTagContainer = EffectDamage1Tag.GetSingleTagContainer();
+		FGameplayTagContainer ParentContainer = EffectDamage1Tag.GetGameplayTagParents();
+
+		TestTrueExpr(SingleTagContainer.HasTagExact(EffectDamage1Tag));
+		TestTrueExpr(SingleTagContainer.HasTag(EffectDamageTag));
+		TestTrueExpr(!SingleTagContainer.HasTagExact(EffectDamageTag));
+
+		TestTrueExpr(ParentContainer.HasTagExact(EffectDamage1Tag));
+		TestTrueExpr(ParentContainer.HasTag(EffectDamageTag));
+		TestTrueExpr(ParentContainer.HasTagExact(EffectDamageTag));
+
 	}
 
 	void GameplayTagTest_PerfTest()
@@ -279,6 +291,22 @@ public:
 				TagContainerNew.AppendTags(TagContainer);
 			}
 		}
+
+		{
+			FScopeLogTime LogTimePtr(*FString::Printf(TEXT("10000 container gets")), nullptr, FScopeLogTime::ScopeLog_Milliseconds);
+			for (int32 i = 0; i < 10000; i++)
+			{
+				FGameplayTagContainer TagContainerNew = EffectDamage1Tag.GetSingleTagContainer();
+			}
+		}
+
+		{
+			FScopeLogTime LogTimePtr(*FString::Printf(TEXT("10000 parent gets")), nullptr, FScopeLogTime::ScopeLog_Milliseconds);
+			for (int32 i = 0; i < 10000; i++)
+			{
+				FGameplayTagContainer TagContainerParents = EffectDamage1Tag.GetGameplayTagParents();
+			}
+		}
 	
 		FGameplayTagContainer TagContainer2;
 		TagContainer2.AddTag(EffectDamage1Tag);
@@ -301,6 +329,14 @@ public:
 			}
 		}
 
+		{
+			FScopeLogTime LogTimePtr(*FString::Printf(TEXT("10000 MatchesTag checks")), nullptr, FScopeLogTime::ScopeLog_Milliseconds);
+			for (int32 i = 0; i < 10000; i++)
+			{
+				bResult &= EffectDamage1Tag.MatchesTag(EffectDamageTag);
+			}
+		}
+
 		{
 			FScopeLogTime LogTimePtr(*FString::Printf(TEXT("10000 HasTagExact checks")), nullptr, FScopeLogTime::ScopeLog_Milliseconds);
 			for (int32 i = 0; i < 10000; i++)
