Added specialization for gameplaytag serialization to FArchive to avoid issue with conflicting template generalization that behaves differently based on config.

diff --git a/Engine/Source/Runtime/GameplayTags/Classes/GameplayTagContainer.h b/Engine/Source/Runtime/GameplayTags/Classes/GameplayTagContainer.h
--- a/Engine/Source/Runtime/GameplayTags/Classes/GameplayTagContainer.h
+++ b/Engine/Source/Runtime/GameplayTags/Classes/GameplayTagContainer.h
@@ -186,6 +186,11 @@ struct FGameplayTag
 		return TagName;
 	}
 
+	friend FArchive& operator<<(FArchive& Ar, FGameplayTag& GameplayTag)
+	{
+		return Ar << GameplayTag.TagName;
+	}
+
 	friend void operator<<(FStructuredArchive::FSlot Slot, FGameplayTag& GameplayTag)
 	{
 		Slot << GameplayTag.TagName;
