Gameplay Tag Manager: Adjusted the loop in "UGameplayTagsManager::StaticGetCategoriesMetaFromPropertyHandle" to break early if the outer base class of a child property and its parent are of a different type. This helps us avoid situations where a suboject's gameplay tag property may inherit meta categories from a different class it didn't expect and cause unintented restrictions on the child tag property.

#tests 28.20 editor

diff --git a/Engine/Source/Runtime/GameplayTags/Private/GameplayTagsManager.cpp b/Engine/Source/Runtime/GameplayTags/Private/GameplayTagsManager.cpp
--- a/Engine/Source/Runtime/GameplayTags/Private/GameplayTagsManager.cpp
+++ b/Engine/Source/Runtime/GameplayTags/Private/GameplayTagsManager.cpp
@@ -1584,7 +1584,20 @@ FString UGameplayTagsManager::StaticGetCategoriesMetaFromPropertyHandle(TSharedP
 				}
 			}
 		}
-		PropertyHandle = PropertyHandle->GetParentHandle();
+
+		TSharedPtr<IPropertyHandle> ParentHandle = PropertyHandle->GetParentHandle();
+		
+		if (ParentHandle.IsValid())
+		{
+			/** Check if the parent handle's base class is of the same class. It's possible the current child property is from a subobject which in that case we probably want to ignore
+			  * any meta category restrictions coming from any parent properties. A subobject's gameplay tag property without any declared meta categories should stay that way. */
+			if (PropertyHandle->GetOuterBaseClass() != ParentHandle->GetOuterBaseClass())
+			{
+				break;
+			}
+		}
+
+		PropertyHandle = ParentHandle;
 	}
 	
 	return Categories;
