Fix ensure while modifying Commonly Replicated Tags and fix double adding that came about because of a bad merge in 2021

diff --git a/Engine/Source/Runtime/GameplayTags/Private/GameplayTagsManager.cpp b/Engine/Source/Runtime/GameplayTags/Private/GameplayTagsManager.cpp
--- a/Engine/Source/Runtime/GameplayTags/Private/GameplayTagsManager.cpp
+++ b/Engine/Source/Runtime/GameplayTags/Private/GameplayTagsManager.cpp
@@ -539,9 +539,13 @@ void UGameplayTagsManager::ConstructGameplayTagTree()
 			CommonlyReplicatedTags.Empty();
 			for (FName TagName : MutableDefault->CommonlyReplicatedTags)
 			{
-				CommonlyReplicatedTags.Add(FGameplayTag(TagName));
+				if (TagName.IsNone())
+				{
+					// Still being added to the UI
+					continue;
+				}
 
-				FGameplayTag Tag = RequestGameplayTag(TagName);
+				FGameplayTag Tag = RequestGameplayTag(TagName, false);
 				if (Tag.IsValid())
 				{
 					CommonlyReplicatedTags.Add(Tag);
