Changed check in FQueryEvaluator::EvalExpr to be use ensureAlwaysMsgf instead, that way it won't just close the executable if it's hit.

#changelist validated
#virtualized

diff --git a/Engine/Source/Runtime/GameplayTags/Private/GameplayTagContainer.cpp b/Engine/Source/Runtime/GameplayTags/Private/GameplayTagContainer.cpp
--- a/Engine/Source/Runtime/GameplayTags/Private/GameplayTagContainer.cpp
+++ b/Engine/Source/Runtime/GameplayTags/Private/GameplayTagContainer.cpp
@@ -503,7 +503,7 @@ bool FQueryEvaluator::EvalExpr(FGameplayTagContainer const& Tags, bool bSkip)
 		return EvalNoExprMatch(Tags, bSkip);
 	}
 
-	check(false);
+	ensureAlwaysMsgf(false, TEXT("Encounted invalid query expression type: %d"), *UEnum::GetValueAsString(ExprType));
 	return false;
 }
 
