Make ExtractParentTags work in transactions.

diff --git a/Engine/Source/Runtime/GameplayTags/Private/GameplayTagsManager.cpp b/Engine/Source/Runtime/GameplayTags/Private/GameplayTagsManager.cpp
--- a/Engine/Source/Runtime/GameplayTags/Private/GameplayTagsManager.cpp
+++ b/Engine/Source/Runtime/GameplayTags/Private/GameplayTagsManager.cpp
@@ -21,6 +21,7 @@
 #include "Misc/CoreDelegates.h"
 #include "HAL/IConsoleManager.h"
 #include "NativeGameplayTags.h"
+#include "AutoRTFM/AutoRTFM.h"
 
 #if WITH_EDITOR
 #include "SourceControlHelpers.h"
@@ -2230,37 +2231,44 @@ bool UGameplayTagsManager::ExtractParentTags(const FGameplayTag& GameplayTag, TA
 		return false;
 	}
 
-#if VALIDATE_EXTRACT_PARENT_TAGS
-	TArray<FGameplayTag> ValidationCopy = UniqueParentTags;
-#endif
+	TArray<FGameplayTag> ValidationCopy;
 
-	FScopeLock Lock(&GameplayTagMapCritical);
+	if constexpr (0 != VALIDATE_EXTRACT_PARENT_TAGS)
+	{
+		ValidationCopy = UniqueParentTags;
+	}
 
 	int32 OldSize = UniqueParentTags.Num();
 	FName RawTag = GameplayTag.GetTagName();
 
-	// This code does not check redirectors because that was already handled by GameplayTagContainerLoaded
-	const TSharedPtr<FGameplayTagNode>* Node = GameplayTagNodeMap.Find(GameplayTag);
-	if (Node)
-	{
-		// Use the registered tag container if it exists
-		const FGameplayTagContainer& SingleContainer = (*Node)->GetSingleTagContainer();
-		for (const FGameplayTag& ParentTag : SingleContainer.ParentTags)
+	// Need to run in the open as it takes a lock. 
+	UE_AUTORTFM_OPEN(
 		{
-			UniqueParentTags.AddUnique(ParentTag);
-		}
+			FScopeLock Lock(&GameplayTagMapCritical);
 
-#if VALIDATE_EXTRACT_PARENT_TAGS
-		GameplayTag.ParseParentTags(ValidationCopy);
+			// This code does not check redirectors because that was already handled by GameplayTagContainerLoaded
+			const TSharedPtr<FGameplayTagNode>*Node = GameplayTagNodeMap.Find(GameplayTag);
+			if (Node)
+			{
+				// Use the registered tag container if it exists
+				const FGameplayTagContainer& SingleContainer = (*Node)->GetSingleTagContainer();
+				for (const FGameplayTag& ParentTag : SingleContainer.ParentTags)
+				{
+					UniqueParentTags.AddUnique(ParentTag);
+				}
 
-		ensureAlwaysMsgf(ValidationCopy == UniqueParentTags, TEXT("ExtractParentTags results are inconsistent for tag %s"), *GameplayTag.ToString());
-#endif
-	}
-	else if (!ShouldClearInvalidTags())
-	{
-		// If we don't clear invalid tags, we need to extract the parents now in case they get registered later
-		GameplayTag.ParseParentTags(UniqueParentTags);
-	}
+				if constexpr (0 != VALIDATE_EXTRACT_PARENT_TAGS)
+				{
+					GameplayTag.ParseParentTags(ValidationCopy);
+					ensureAlwaysMsgf(ValidationCopy == UniqueParentTags, TEXT("ExtractParentTags results are inconsistent for tag %s"), *GameplayTag.ToString());
+				}
+			}
+			else if (!ShouldClearInvalidTags())
+			{
+				// If we don't clear invalid tags, we need to extract the parents now in case they get registered later
+				GameplayTag.ParseParentTags(UniqueParentTags);
+			}
+		});
 
 	return UniqueParentTags.Num() != OldSize;
 }
