Added a way to temporarily override how the GameplayTagManager handles unloading of GameplayTags.

#tests Verified in cooked Creative with a single client traveling back and forth between two different  servers, observing how the NetworkGameplayTagNodeIndexHash change after GFS deacivation.

diff --git a/Engine/Source/Runtime/GameplayTags/Classes/GameplayTagsManager.h b/Engine/Source/Runtime/GameplayTags/Classes/GameplayTagsManager.h
--- a/Engine/Source/Runtime/GameplayTags/Classes/GameplayTagsManager.h
+++ b/Engine/Source/Runtime/GameplayTags/Classes/GameplayTagsManager.h
@@ -554,6 +554,12 @@ public:
 	/** If we are allowed to unload tags */
 	GAMEPLAYTAGS_API bool ShouldUnloadTags() const;
 
+	/** Pushes an override that supersedes bShouldAllowUnloadingTags to allow/disallow unloading of GameplayTags in controlled scenarios */
+	GAMEPLAYTAGS_API void SetShouldUnloadTagsOverride(bool bShouldUnloadTags);
+
+	/** Clears runtime overrides, reverting to bShouldAllowUnloadingTags when determining GameplayTags unload behavior */
+	GAMEPLAYTAGS_API void ClearShouldUnloadTagsOverride();
+
 	/** Returns the hash of NetworkGameplayTagNodeIndex */
 	uint32 GetNetworkGameplayTagNodeIndexHash() const { VerifyNetworkIndex(); return NetworkGameplayTagNodeIndexHash; }
 
@@ -885,6 +891,9 @@ private:
 	/** Cached runtime value for whether we should allow unloading of tags */
 	bool bShouldAllowUnloadingTags;
 
+	/** Augments usage of bShouldAllowUnloadingTags to allow runtime overrides to allow/disallow unloading of GameplayTags in controlled scenarios */
+	TOptional<bool> ShouldAllowUnloadingTagsOverride;
+
 	/** True if native tags have all been added and flushed */
 	bool bDoneAddingNativeTags;
 
diff --git a/Engine/Source/Runtime/GameplayTags/Private/GameplayTagsManager.cpp b/Engine/Source/Runtime/GameplayTags/Private/GameplayTagsManager.cpp
--- a/Engine/Source/Runtime/GameplayTags/Private/GameplayTagsManager.cpp
+++ b/Engine/Source/Runtime/GameplayTags/Private/GameplayTagsManager.cpp
@@ -740,9 +740,24 @@ bool UGameplayTagsManager::ShouldUnloadTags() const
 	}
 #endif
 
+	if (ShouldAllowUnloadingTagsOverride.IsSet())
+	{
+		return ShouldAllowUnloadingTagsOverride.GetValue();
+	}
+
 	return bShouldAllowUnloadingTags;
 }
 
+void UGameplayTagsManager::SetShouldUnloadTagsOverride(bool bShouldUnloadTags)
+{
+	ShouldAllowUnloadingTagsOverride = bShouldUnloadTags;
+}
+
+void UGameplayTagsManager::ClearShouldUnloadTagsOverride()
+{
+	ShouldAllowUnloadingTagsOverride.Reset();
+}
+
 void UGameplayTagsManager::GetRestrictedTagConfigFiles(TArray<FString>& RestrictedConfigFiles) const
 {
 	UGameplayTagsSettings* MutableDefault = GetMutableDefault<UGameplayTagsSettings>();
