Added a way to temporarily defer rebuilds of the GameplayTag tree managed by the GameplayTagManager.
Also added a way to optionally rebuild the tree when un-setting the corresponding override.

#tests Verified in cooked Creative with a single client cycling between servers, observing timing and final NetworkGameplayTagNodeIndexHash values.

diff --git a/Engine/Source/Runtime/GameplayTags/Classes/GameplayTagsManager.h b/Engine/Source/Runtime/GameplayTags/Classes/GameplayTagsManager.h
--- a/Engine/Source/Runtime/GameplayTags/Classes/GameplayTagsManager.h
+++ b/Engine/Source/Runtime/GameplayTags/Classes/GameplayTagsManager.h
@@ -560,6 +560,12 @@ public:
 	/** Clears runtime overrides, reverting to bShouldAllowUnloadingTags when determining GameplayTags unload behavior */
 	GAMEPLAYTAGS_API void ClearShouldUnloadTagsOverride();
 
+	/** Pushes an override that suppresses calls to HandleGameplayTagTreeChanged that would result in a complete rebuild of the GameplayTag tree */
+	GAMEPLAYTAGS_API void SetShouldDeferGameplayTagTreeRebuilds(bool bShouldDeferRebuilds);
+
+	/** Stops suppressing GameplayTag tree rebuilds and (optionally) rebuilds the tree */
+	GAMEPLAYTAGS_API void ClearShouldDeferGameplayTagTreeRebuilds(bool bRebuildTree);
+
 	/** Returns the hash of NetworkGameplayTagNodeIndex */
 	uint32 GetNetworkGameplayTagNodeIndexHash() const { VerifyNetworkIndex(); return NetworkGameplayTagNodeIndexHash; }
 
@@ -894,6 +900,9 @@ private:
 	/** Augments usage of bShouldAllowUnloadingTags to allow runtime overrides to allow/disallow unloading of GameplayTags in controlled scenarios */
 	TOptional<bool> ShouldAllowUnloadingTagsOverride;
 
+	/** Used to suppress calls to HandleGameplayTagTreeChanged that would result in a complete rebuild of the GameplayTag tree*/
+	TOptional<bool> ShouldDeferGameplayTagTreeRebuilds;
+
 	/** True if native tags have all been added and flushed */
 	bool bDoneAddingNativeTags;
 
diff --git a/Engine/Source/Runtime/GameplayTags/Private/GameplayTagsManager.cpp b/Engine/Source/Runtime/GameplayTags/Private/GameplayTagsManager.cpp
--- a/Engine/Source/Runtime/GameplayTags/Private/GameplayTagsManager.cpp
+++ b/Engine/Source/Runtime/GameplayTags/Private/GameplayTagsManager.cpp
@@ -758,6 +758,21 @@ void UGameplayTagsManager::ClearShouldUnloadTagsOverride()
 	ShouldAllowUnloadingTagsOverride.Reset();
 }
 
+void UGameplayTagsManager::SetShouldDeferGameplayTagTreeRebuilds(bool bShouldDeferRebuilds)
+{
+	ShouldDeferGameplayTagTreeRebuilds = bShouldDeferRebuilds;
+}
+
+void UGameplayTagsManager::ClearShouldDeferGameplayTagTreeRebuilds(bool bRebuildTree)
+{
+	ShouldDeferGameplayTagTreeRebuilds.Reset();
+
+	if (bRebuildTree)
+	{
+		HandleGameplayTagTreeChanged(true);
+	}
+}
+
 void UGameplayTagsManager::GetRestrictedTagConfigFiles(TArray<FString>& RestrictedConfigFiles) const
 {
 	UGameplayTagsSettings* MutableDefault = GetMutableDefault<UGameplayTagsSettings>();
@@ -1185,7 +1200,7 @@ void UGameplayTagsManager::HandleGameplayTagTreeChanged(bool bRecreateTree)
 	// Don't do anything during a reconstruct or before initial native tags are done loading
 	if (!bIsConstructingGameplayTagTree && bDoneAddingNativeTags)
 	{
-		if (bRecreateTree)
+		if (bRecreateTree && (!ShouldDeferGameplayTagTreeRebuilds.IsSet() || !ShouldDeferGameplayTagTreeRebuilds.GetValue()))
 		{
 #if WITH_EDITOR
 			if (GIsEditor)
