Optimize UGameplayTagsManager::GameplayTagsMatchDepth to remove excess allocations.
* Kept old behavior for validation for now

#tests Been using it locally for some time and frontend locker makes use of this functionality

diff --git a/Engine/Source/Runtime/GameplayTags/Private/GameplayTagsManager.cpp b/Engine/Source/Runtime/GameplayTags/Private/GameplayTagsManager.cpp
--- a/Engine/Source/Runtime/GameplayTags/Private/GameplayTagsManager.cpp
+++ b/Engine/Source/Runtime/GameplayTags/Private/GameplayTagsManager.cpp
@@ -2388,22 +2388,39 @@ void UGameplayTagsManager::SplitGameplayTagFName(const FGameplayTag& Tag, TArray
 
 int32 UGameplayTagsManager::GameplayTagsMatchDepth(const FGameplayTag& GameplayTagOne, const FGameplayTag& GameplayTagTwo) const
 {
-	TSet<FName> Tags1;
-	TSet<FName> Tags2;
+	using FTagsArray = TArray<FName, TInlineAllocator<32>>;
 
-	TSharedPtr<FGameplayTagNode> TagNode = FindTagNode(GameplayTagOne);
-	if (TagNode.IsValid())
+	auto GetTags = [this](FTagsArray& Tags, const FGameplayTag& GameplayTag)
 	{
-		GetAllParentNodeNames(Tags1, TagNode);
-	}
+		for (TSharedPtr<FGameplayTagNode> TagNode = FindTagNode(GameplayTag); TagNode.IsValid(); TagNode = TagNode->GetParentTagNode())
+		{
+			Tags.Add(TagNode->Tag);
+		}
+	};
+
+	FTagsArray Tags1, Tags2;
+	GetTags(Tags1, GameplayTagOne);
+	GetTags(Tags2, GameplayTagTwo);
+
+	// Get Tags returns tail to head, so compare in reverse order
+	int32 Index1 = Tags1.Num() - 1;
+	int32 Index2 = Tags2.Num() - 1;
+
+	int32 Depth = 0;
 
-	TagNode = FindTagNode(GameplayTagTwo);
-	if (TagNode.IsValid())
+	for (; Index1 >= 0 && Index2 >= 0; --Index1, --Index2)
 	{
-		GetAllParentNodeNames(Tags2, TagNode);
+		if (Tags1[Index1] == Tags2[Index2])
+		{
+			++Depth;
+		}
+		else
+		{
+			break;
+		}
 	}
 
-	return Tags1.Intersect(Tags2).Num();
+	return Depth;
 }
 
 int32 UGameplayTagsManager::GetNumberOfTagNodes(const FGameplayTag& GameplayTag) const
