Optimize FGameplayTagContainer::AppendTags to be much faster when appending large containers to small containers.
Adjust tag perf tests to cover this case, I validated the perf improvement on a Test build.

diff --git a/Engine/Source/Runtime/GameplayTags/Private/GameplayTagContainer.cpp b/Engine/Source/Runtime/GameplayTags/Private/GameplayTagContainer.cpp
--- a/Engine/Source/Runtime/GameplayTags/Private/GameplayTagContainer.cpp
+++ b/Engine/Source/Runtime/GameplayTags/Private/GameplayTagContainer.cpp
@@ -667,18 +667,50 @@ void FGameplayTagContainer::AppendTags(FGameplayTagContainer const& Other)
 {
 	CONDITIONAL_SCOPE_CYCLE_COUNTER(STAT_FGameplayTagContainer_AppendTags, GEnableGameplayTagDetailedStats);
 
-	GameplayTags.Reserve(GameplayTags.Num() + Other.GameplayTags.Num());
-	ParentTags.Reserve(ParentTags.Num() + Other.ParentTags.Num());
-
+	int32 OldTagNum = GameplayTags.Num();
+	GameplayTags.Reserve(OldTagNum + Other.GameplayTags.Num());
 	// Add other container's tags to our own
 	for(const FGameplayTag& OtherTag : Other.GameplayTags)
 	{
-		GameplayTags.AddUnique(OtherTag);
+		int32 SearchIndex = 0;
+		while (true)
+		{
+			if (SearchIndex >= OldTagNum)
+			{
+				// Stop searching once we've looked at all existing tags, this is faster when appending large containers
+				GameplayTags.Add(OtherTag);
+				break;
+			}
+			else if (GameplayTags[SearchIndex] == OtherTag)
+			{
+				// Matching tag found, stop searching
+				break;
+			}
+
+			SearchIndex++;
+		}
 	}
 
+	// This function is called enough that the code duplication is faster than a lambda
+	OldTagNum = ParentTags.Num();
+	ParentTags.Reserve(OldTagNum + Other.ParentTags.Num());
 	for (const FGameplayTag& OtherTag : Other.ParentTags)
 	{
-		ParentTags.AddUnique(OtherTag);
+		int32 SearchIndex = 0;
+		while (true)
+		{
+			if (SearchIndex >= OldTagNum)
+			{
+				ParentTags.Add(OtherTag);
+				break;
+			}
+			else if (ParentTags[SearchIndex] == OtherTag)
+			{
+				break;
+			}
+
+			SearchIndex++;
+		}
 	}
 }
 
diff --git a/Engine/Source/Runtime/GameplayTags/Private/Tests/GameplayTagTests.cpp b/Engine/Source/Runtime/GameplayTags/Private/Tests/GameplayTagTests.cpp
--- a/Engine/Source/Runtime/GameplayTags/Private/Tests/GameplayTagTests.cpp
+++ b/Engine/Source/Runtime/GameplayTags/Private/Tests/GameplayTagTests.cpp
@@ -9,7 +9,7 @@
 #include "GameplayTagsModule.h"
 #include "Stats/StatsMisc.h"
 
-#if WITH_DEV_AUTOMATION_TESTS
+#if WITH_AUTOMATION_WORKER
 
 class FGameplayTagTestBase : public FAutomationTestBase
 {
@@ -245,18 +245,19 @@ public:
 		FGameplayTagContainer TagContainer;
 
 		bool bResult = true;
+		int32 SmallTest = 1000, LargeTest = 10000;
 
 		{
-			FScopeLogTime LogTimePtr(*FString::Printf(TEXT("10000 get tag")), nullptr, FScopeLogTime::ScopeLog_Milliseconds);
-			for (int32 i = 0; i < 10000; i++)
+			FScopeLogTime LogTimePtr(*FString::Printf(TEXT("%d get tag"), LargeTest), nullptr, FScopeLogTime::ScopeLog_Milliseconds);
+			for (int32 i = 0; i < LargeTest; i++)
 			{
 				UGameplayTagsManager::Get().RequestGameplayTag(FName(TEXT("Effect.Damage")));
 			}
 		}
 
 		{
-			FScopeLogTime LogTimePtr(*FString::Printf(TEXT("1000 container constructions")), nullptr, FScopeLogTime::ScopeLog_Milliseconds);
-			for (int32 i = 0; i < 1000; i++)
+			FScopeLogTime LogTimePtr(*FString::Printf(TEXT("%d container constructions"), SmallTest), nullptr, FScopeLogTime::ScopeLog_Milliseconds);
+			for (int32 i = 0; i < SmallTest; i++)
 			{
 				TagContainer = FGameplayTagContainer();
 				TagContainer.AddTag(EffectDamage1Tag);
@@ -270,8 +271,21 @@ public:
 		}
 
 		{
-			FScopeLogTime LogTimePtr(*FString::Printf(TEXT("1000 container copies")), nullptr, FScopeLogTime::ScopeLog_Milliseconds);
-			for (int32 i = 0; i < 1000; i++)
+			FScopeLogTime LogTimePtr(*FString::Printf(TEXT("%d container copy and move"), SmallTest), nullptr, FScopeLogTime::ScopeLog_Milliseconds);
+			for (int32 i = 0; i < SmallTest; i++)
+			{
+				FGameplayTagContainer TagContainerNew(EffectDamageTag);
+				TagContainerNew = TagContainer;
+
+				FGameplayTagContainer MovedContainer = MoveTemp(TagContainerNew);
+
+				bResult &= (MovedContainer.Num() == TagContainer.Num());
+			}
+		}
+
+		{
+			FScopeLogTime LogTimePtr(*FString::Printf(TEXT("%d container addtag"), SmallTest), nullptr, FScopeLogTime::ScopeLog_Milliseconds);
+			for (int32 i = 0; i < SmallTest; i++)
 			{
 				FGameplayTagContainer TagContainerNew;
 
@@ -283,26 +297,36 @@ public:
 		}
 
 		{
-			FScopeLogTime LogTimePtr(*FString::Printf(TEXT("1000 container appends")), nullptr, FScopeLogTime::ScopeLog_Milliseconds);
-			for (int32 i = 0; i < 1000; i++)
+			FScopeLogTime LogTimePtr(*FString::Printf(TEXT("%d container partial appends"), SmallTest), nullptr, FScopeLogTime::ScopeLog_Milliseconds);
+			for (int32 i = 0; i < SmallTest; i++)
 			{
-				FGameplayTagContainer TagContainerNew;
+				FGameplayTagContainer TagContainerNew(EffectDamage1Tag);
+
+				TagContainerNew.AppendTags(TagContainer);
+			}
+		}
+
+		{
+			FScopeLogTime LogTimePtr(*FString::Printf(TEXT("%d container full appends"), SmallTest), nullptr, FScopeLogTime::ScopeLog_Milliseconds);
+			for (int32 i = 0; i < SmallTest; i++)
+			{
+				FGameplayTagContainer TagContainerNew = TagContainer;
 
 				TagContainerNew.AppendTags(TagContainer);
 			}
 		}
 
 		{
-			FScopeLogTime LogTimePtr(*FString::Printf(TEXT("10000 container gets")), nullptr, FScopeLogTime::ScopeLog_Milliseconds);
-			for (int32 i = 0; i < 10000; i++)
+			FScopeLogTime LogTimePtr(*FString::Printf(TEXT("%d container gets"), LargeTest), nullptr, FScopeLogTime::ScopeLog_Milliseconds);
+			for (int32 i = 0; i < LargeTest; i++)
 			{
 				FGameplayTagContainer TagContainerNew = EffectDamage1Tag.GetSingleTagContainer();
 			}
 		}
 
 		{
-			FScopeLogTime LogTimePtr(*FString::Printf(TEXT("10000 parent gets")), nullptr, FScopeLogTime::ScopeLog_Milliseconds);
-			for (int32 i = 0; i < 10000; i++)
+			FScopeLogTime LogTimePtr(*FString::Printf(TEXT("%d parent gets"), LargeTest), nullptr, FScopeLogTime::ScopeLog_Milliseconds);
+			for (int32 i = 0; i < LargeTest; i++)
 			{
 				FGameplayTagContainer TagContainerParents = EffectDamage1Tag.GetGameplayTagParents();
 			}
@@ -314,56 +338,56 @@ public:
 		TagContainer2.AddTag(CueTag);
 
 		{
-			FScopeLogTime LogTimePtr(*FString::Printf(TEXT("10000 MatchesAnyExact checks")), nullptr, FScopeLogTime::ScopeLog_Milliseconds);
-			for (int32 i = 0; i < 10000; i++)
+			FScopeLogTime LogTimePtr(*FString::Printf(TEXT("%d MatchesAnyExact checks"), LargeTest), nullptr, FScopeLogTime::ScopeLog_Milliseconds);
+			for (int32 i = 0; i < LargeTest; i++)
 			{
 				bResult &= EffectDamage1Tag.MatchesAnyExact(TagContainer);
 			}
 		}
 
 		{
-			FScopeLogTime LogTimePtr(*FString::Printf(TEXT("10000 MatchesAny checks")), nullptr, FScopeLogTime::ScopeLog_Milliseconds);
-			for (int32 i = 0; i < 10000; i++)
+			FScopeLogTime LogTimePtr(*FString::Printf(TEXT("%d MatchesAny checks"), LargeTest), nullptr, FScopeLogTime::ScopeLog_Milliseconds);
+			for (int32 i = 0; i < LargeTest; i++)
 			{
 				bResult &= EffectDamage1Tag.MatchesAny(TagContainer);
 			}
 		}
 
 		{
-			FScopeLogTime LogTimePtr(*FString::Printf(TEXT("10000 MatchesTag checks")), nullptr, FScopeLogTime::ScopeLog_Milliseconds);
-			for (int32 i = 0; i < 10000; i++)
+			FScopeLogTime LogTimePtr(*FString::Printf(TEXT("%d MatchesTag checks"), LargeTest), nullptr, FScopeLogTime::ScopeLog_Milliseconds);
+			for (int32 i = 0; i < LargeTest; i++)
 			{
 				bResult &= EffectDamage1Tag.MatchesTag(EffectDamageTag);
 			}
 		}
 
 		{
-			FScopeLogTime LogTimePtr(*FString::Printf(TEXT("10000 HasTagExact checks")), nullptr, FScopeLogTime::ScopeLog_Milliseconds);
-			for (int32 i = 0; i < 10000; i++)
+			FScopeLogTime LogTimePtr(*FString::Printf(TEXT("%d HasTagExact checks"), LargeTest), nullptr, FScopeLogTime::ScopeLog_Milliseconds);
+			for (int32 i = 0; i < LargeTest; i++)
 			{
 				bResult &= TagContainer.HasTagExact(EffectDamage1Tag);
 			}
 		}
 
 		{
-			FScopeLogTime LogTimePtr(*FString::Printf(TEXT("10000 HasTag checks")), nullptr, FScopeLogTime::ScopeLog_Milliseconds);
-			for (int32 i = 0; i < 10000; i++)
+			FScopeLogTime LogTimePtr(*FString::Printf(TEXT("%d HasTag checks"), LargeTest), nullptr, FScopeLogTime::ScopeLog_Milliseconds);
+			for (int32 i = 0; i < LargeTest; i++)
 			{
 				bResult &= TagContainer.HasTag(EffectDamage1Tag);
 			}
 		}
 
 		{
-			FScopeLogTime LogTimePtr(*FString::Printf(TEXT("10000 HasAll checks")), nullptr, FScopeLogTime::ScopeLog_Milliseconds);
-			for (int32 i = 0; i < 10000; i++)
+			FScopeLogTime LogTimePtr(*FString::Printf(TEXT("%d HasAll checks"), LargeTest), nullptr, FScopeLogTime::ScopeLog_Milliseconds);
+			for (int32 i = 0; i < LargeTest; i++)
 			{
 				bResult &= TagContainer.HasAll(TagContainer2);
 			}
 		}
 
 		{
-			FScopeLogTime LogTimePtr(*FString::Printf(TEXT("10000 HasAny checks")), nullptr, FScopeLogTime::ScopeLog_Milliseconds);
-			for (int32 i = 0; i < 10000; i++)
+			FScopeLogTime LogTimePtr(*FString::Printf(TEXT("%d HasAny checks"), LargeTest), nullptr, FScopeLogTime::ScopeLog_Milliseconds);
+			for (int32 i = 0; i < LargeTest; i++)
 			{
 				bResult &= TagContainer.HasAny(TagContainer2);
 			}
@@ -391,4 +415,4 @@ bool FGameplayTagTest::RunTest(const FString& Parameters)
 	return !HasAnyErrors();
 }
 
-#endif //WITH_DEV_AUTOMATION_TESTS
+#endif //WITH_AUTOMATION_WORKER
