Adding support to validate tags with assets instead of editor propery handles

diff --git a/Engine/Source/Runtime/GameplayTags/Classes/GameplayTagsManager.h b/Engine/Source/Runtime/GameplayTags/Classes/GameplayTagsManager.h
--- a/Engine/Source/Runtime/GameplayTags/Classes/GameplayTagsManager.h
+++ b/Engine/Source/Runtime/GameplayTags/Classes/GameplayTagsManager.h
@@ -703,17 +703,28 @@ public:
 	DECLARE_MULTICAST_DELEGATE_ThreeParams(FOnFilterGameplayTagChildren, const FString&  /** FilterString */, TSharedPtr<FGameplayTagNode>& /* TagNode */, bool& /* OUT OutShouldHide */)
 	FOnFilterGameplayTagChildren OnFilterGameplayTagChildren;
 
+	/*
+	* This is a container to filter out gameplay tags when they are invalid or when they don't meet the filter string
+	* If used from editor to filter out tags when picking them the FilterString is optional and the ReferencingPropertyHandle is required
+	* If used to validate an asset / assets you can provide the TagSourceAssets. The FilterString is optional and the ReferencingPropertyHandle is not required
+	*/
 	struct FFilterGameplayTagContext
 	{
 		const FString& FilterString;
 		const TSharedPtr<FGameplayTagNode>& TagNode;
 		const FGameplayTagSource* TagSource;
 		const TSharedPtr<IPropertyHandle>& ReferencingPropertyHandle;
+		const TArray<FAssetData> TagSourceAssets;
 
 		FFilterGameplayTagContext(const FString& InFilterString, const TSharedPtr<FGameplayTagNode>& InTagNode, const FGameplayTagSource* InTagSource, const TSharedPtr<IPropertyHandle>& InReferencingPropertyHandle)
 			: FilterString(InFilterString), TagNode(InTagNode), TagSource(InTagSource), ReferencingPropertyHandle(InReferencingPropertyHandle)
 		{}
+
+		FFilterGameplayTagContext(const TSharedPtr<FGameplayTagNode>& InTagNode, const FGameplayTagSource* InTagSource, const TArray<FAssetData>& InTagSourceAssets, const FString& InFilterString = FString())
+			: FilterString(InFilterString), TagNode(InTagNode), TagSource(InTagSource), ReferencingPropertyHandle(nullptr), TagSourceAssets(InTagSourceAssets)
+		{}
 	};
+
 	/*
 	 * Allows dynamic hiding of gameplay tags in SGameplayTagWidget. Allows higher order structs to dynamically change which tags are visible based on its own data
 	 * Applies to all tags, and has more context than OnFilterGameplayTagChildren
