Fixing cis issue
-Adding missing include for FAssetData

diff --git a/Engine/Source/Runtime/GameplayTags/Classes/GameplayTagsManager.h b/Engine/Source/Runtime/GameplayTags/Classes/GameplayTagsManager.h
--- a/Engine/Source/Runtime/GameplayTags/Classes/GameplayTagsManager.h
+++ b/Engine/Source/Runtime/GameplayTags/Classes/GameplayTagsManager.h
@@ -2,6 +2,7 @@
 
 #pragma once
 
+#include "AssetRegistry/AssetData.h"
 #include "CoreMinimal.h"
 #include "Stats/Stats.h"
 #include "UObject/ObjectMacros.h"
