Added public native function for setting the user description of a FGameplayTagQuery.

#11510

diff --git a/Engine/Source/Runtime/GameplayTags/Classes/GameplayTagContainer.h b/Engine/Source/Runtime/GameplayTags/Classes/GameplayTagContainer.h
--- a/Engine/Source/Runtime/GameplayTags/Classes/GameplayTagContainer.h
+++ b/Engine/Source/Runtime/GameplayTags/Classes/GameplayTagContainer.h
@@ -883,6 +883,9 @@ public:
 	/** Serialize the tag query */
 	GAMEPLAYTAGS_API void Serialize(FArchive& Ar);
 	
+	/** Sets the user-specified description of this query. */
+	void SetUserDescription(const FString& InUserDescription) { UserDescription = InUserDescription; }
+
 	/** Returns description string. */
 	const FString& GetDescription() const { return UserDescription.IsEmpty() ? AutoDescription : UserDescription; };
 
