Always open around getting the gameplay tag parents

diff --git a/Engine/Source/Runtime/GameplayTags/Private/GameplayTagsManager.cpp b/Engine/Source/Runtime/GameplayTags/Private/GameplayTagsManager.cpp
--- a/Engine/Source/Runtime/GameplayTags/Private/GameplayTagsManager.cpp
+++ b/Engine/Source/Runtime/GameplayTags/Private/GameplayTagsManager.cpp
@@ -2216,7 +2216,9 @@ void UGameplayTagsManager::DoneAddingNativeTags()
 	}
 }
 
-FGameplayTagContainer UGameplayTagsManager::RequestGameplayTagParents(const FGameplayTag& GameplayTag) const
+// #jira SOL-6589: remove forcenoinline once this JIRA is fixed.
+UE_AUTORTFM_ALWAYS_OPEN
+FORCENOINLINE FGameplayTagContainer UGameplayTagsManager::RequestGameplayTagParents(const FGameplayTag& GameplayTag) const
 {
 	FScopeLock Lock(&GameplayTagMapCritical);
 
