PR #11793: Typo Fix: Update GameplayTagContainer.h

diff --git a/Engine/Source/Runtime/GameplayTags/Classes/GameplayTagContainer.h b/Engine/Source/Runtime/GameplayTags/Classes/GameplayTagContainer.h
--- a/Engine/Source/Runtime/GameplayTags/Classes/GameplayTagContainer.h
+++ b/Engine/Source/Runtime/GameplayTags/Classes/GameplayTagContainer.h
@@ -94,7 +94,7 @@ struct FGameplayTag
 	}
 
 	/**
-	 * Determine if this tag matches TagToCheck, expanding our parent tags
+	 * Determine if this tag matches TagToCheck, expanding out parent tags
 	 * "A.1".MatchesTag("A") will return True, "A".MatchesTag("A.1") will return False
 	 * If TagToCheck is not Valid it will always return False
 	 * 
@@ -115,7 +115,7 @@ struct FGameplayTag
 		{
 			return false;
 		}
-		// Only check check explicit tag list
+		// Only check explicit tag list
 		return TagName == TagToCheck.TagName;
 	}
 
