Make GameplayTagsManager use the transactionally safe lock under AutoRTFM to fix the MatchesTag bug.

diff --git a/Engine/Source/Runtime/Core/Public/Misc/TransactionallySafeCriticalSection.h b/Engine/Source/Runtime/Core/Public/Misc/TransactionallySafeCriticalSection.h
new file mode 100644
--- /dev/null
+++ b/Engine/Source/Runtime/Core/Public/Misc/TransactionallySafeCriticalSection.h
@@ -0,0 +1,92 @@
+// Copyright Epic Games, Inc. All Rights Reserved.
+
+#pragma once
+
+#include "HAL/CriticalSection.h"
+#include "AutoRTFM/AutoRTFM.h"
+
+// A transactionally safe critical section that works in the following novel ways:
+// - In the open (non-transactional):
+//   - Take the lock like before. Simple!
+//   - Free the lock like before too.
+// - In the closed (transactional):
+//   - During locking we query `TransactionalLockCount`:
+//     - 0 means we haven't taken the lock within our transaction nest and need to acquire the lock.
+//     - Otherwise we already have the lock (and are preventing non-transactional code seeing any
+//       modifications we've made while holding the lock), so just bump `TransactionalLockCount`.
+//     - We also register an on-abort handler to release the lock should we abort (but we need to
+//       query `TransactionalLockCount` even there because we could be aborting an inner transaction
+//       and the parent transaction still wants to have the lock held!).
+//   - During unlocking we defer doing the unlock until the transaction commits.
+//
+// Thus with this approach we will hold this lock for the *entirety* of the transactional nest should
+// we take the lock during the transaction, thus preventing non-transactional code from seeing any
+// modifications we should make.
+struct FTransactionallySafeCriticalSectionDefinition final
+{
+	void Lock()
+	{
+		if (AutoRTFM::IsTransactional())
+		{
+			AutoRTFM::Open([&]
+				{
+					// The transactional system which can increment TransactionalLockCount
+					// is always single-threaded, thus this is safe to check without atomicity.
+					if (0 == TransactionalLockCount)
+					{
+						CriticalSection.Lock();
+					}
+
+					TransactionalLockCount += 1;
+				});
+
+			AutoRTFM::OnAbort([this]
+				{
+					check(0 != TransactionalLockCount);
+					TransactionalLockCount -= 1;
+
+					if (0 == TransactionalLockCount)
+					{
+						CriticalSection.Unlock();
+					}
+				});
+		}
+		else
+		{
+			CriticalSection.Lock();
+			check(0 == TransactionalLockCount);
+		}
+	}
+
+	void Unlock()
+	{
+		if (AutoRTFM::IsTransactional())
+		{
+			AutoRTFM::OnCommit([this]
+				{
+					check(0 != TransactionalLockCount);
+					TransactionalLockCount -= 1;
+
+					if (0 == TransactionalLockCount)
+					{
+						CriticalSection.Unlock();
+					}
+				});
+		}
+		else
+		{
+			check(0 == TransactionalLockCount);
+			CriticalSection.Unlock();
+		}
+	}
+
+private:
+	FCriticalSection CriticalSection;
+	uint32 TransactionalLockCount = 0;
+};
+
+#if UE_AUTORTFM
+using FTransactionallySafeCriticalSection = FTransactionallySafeCriticalSectionDefinition;
+#else
+using FTransactionallySafeCriticalSection = FCriticalSection;
+#endif
diff --git a/Engine/Source/Runtime/Core/Public/Misc/TransactionallySafeScopeLock.h b/Engine/Source/Runtime/Core/Public/Misc/TransactionallySafeScopeLock.h
new file mode 100644
--- /dev/null
+++ b/Engine/Source/Runtime/Core/Public/Misc/TransactionallySafeScopeLock.h
@@ -0,0 +1,36 @@
+// Copyright Epic Games, Inc. All Rights Reserved.
+
+#pragma once
+
+#include "TransactionallySafeCriticalSection.h"
+
+// A transactionally safe scope lock - that uses a transactionally safe critical section to back it.
+struct FTransactionallySafeScopeLock final
+{
+	UE_NODISCARD_CTOR FTransactionallySafeScopeLock(FTransactionallySafeCriticalSection* InSynchObject)
+		: SynchObject(InSynchObject)
+	{
+		check(SynchObject);
+		SynchObject->Lock();
+	}
+
+	~FTransactionallySafeScopeLock()
+	{
+		Unlock();
+	}
+
+	void Unlock()
+	{
+		if (SynchObject)
+		{
+			SynchObject->Unlock();
+			SynchObject = nullptr;
+		}
+	}
+
+private:
+	FTransactionallySafeScopeLock() = delete;
+	FTransactionallySafeScopeLock(const FTransactionallySafeScopeLock&) = delete;
+
+	FTransactionallySafeCriticalSection* SynchObject;
+};
diff --git a/Engine/Source/Runtime/Engine/Private/Tests/AutoRTFM/AutoRTFMGameplayTagTests.cpp b/Engine/Source/Runtime/Engine/Private/Tests/AutoRTFM/AutoRTFMGameplayTagTests.cpp
new file mode 100644
--- /dev/null
+++ b/Engine/Source/Runtime/Engine/Private/Tests/AutoRTFM/AutoRTFMGameplayTagTests.cpp
@@ -0,0 +1,67 @@
+// Copyright Epic Games, Inc. All Rights Reserved.
+
+#include "Misc/AutomationTest.h"
+#include "Misc/App.h"
+#include "AutoRTFM/AutoRTFM.h"
+#include "GameplayTagsManager.h"
+#include "GameplayTagContainer.h"
+
+#if WITH_DEV_AUTOMATION_TESTS
+
+#define TEST_CHECK_TRUE(b) do \
+{ \
+	if (!(b)) \
+	{ \
+		FString String(FString::Printf(TEXT("FAILED: %s:%u"), TEXT(__FILE__), __LINE__)); \
+		ExecutionInfo.AddEvent(FAutomationEvent(EAutomationEventType::Info, String)); \
+		return false; \
+	} \
+} while(false)
+
+IMPLEMENT_SIMPLE_AUTOMATION_TEST(FAutoRTFMGameplayTagTests, "AutoRTFM + FGameplayTag", EAutomationTestFlags::EngineFilter | EAutomationTestFlags::ClientContext | EAutomationTestFlags::ServerContext | EAutomationTestFlags::CommandletContext)
+
+bool FAutoRTFMGameplayTagTests::RunTest(const FString & Parameters)
+{
+	if (!AutoRTFM::ForTheRuntime::IsAutoRTFMRuntimeEnabled())
+	{
+		ExecutionInfo.AddEvent(FAutomationEvent(EAutomationEventType::Info, TEXT("SKIPPED 'FAutoRTFMGameplayTagTests' test. AutoRTFM disabled.")));
+		return true;
+	}
+
+	FGameplayTag Tag;
+	FGameplayTag Other;
+
+	bool bResult = true;
+
+	AutoRTFM::ETransactionResult Result = AutoRTFM::Transact([&]
+		{
+			bResult = Tag.MatchesTag(Other);
+		});
+
+	TEST_CHECK_TRUE(AutoRTFM::ETransactionResult::Committed == Result);
+	TEST_CHECK_TRUE(!bResult);
+
+	TArray<FGameplayTag> Parents;
+
+	Result = AutoRTFM::Transact([&]
+		{
+			UGameplayTagsManager::Get().ExtractParentTags(Tag, Parents);
+		});
+
+	TEST_CHECK_TRUE(AutoRTFM::ETransactionResult::Committed == Result);
+	TEST_CHECK_TRUE(Parents.IsEmpty());
+
+	Result = AutoRTFM::Transact([&]
+		{
+			bResult = UGameplayTagsManager::Get().RequestGameplayTagParents(Tag).IsEmpty();
+		});
+
+	TEST_CHECK_TRUE(AutoRTFM::ETransactionResult::Committed == Result);
+	TEST_CHECK_TRUE(bResult);
+	
+	return true;
+}
+
+#undef TEST_CHECK_TRUE
+
+#endif //WITH_DEV_AUTOMATION_TESTS
diff --git a/Engine/Source/Runtime/Engine/Private/TraceFilter.cpp b/Engine/Source/Runtime/Engine/Private/TraceFilter.cpp
--- a/Engine/Source/Runtime/Engine/Private/TraceFilter.cpp
+++ b/Engine/Source/Runtime/Engine/Private/TraceFilter.cpp
@@ -8,115 +8,7 @@
 
 #include "ObjectTrace.h"
 
-// A transactionally safe critical section that works in the following novel ways:
-// - In the open (non-transactional):
-//   - Take the lock like before. Simple!
-//   - Free the lock like before too.
-// - In the closed (transactional):
-//   - During locking we query `TransactionalLockCount`:
-//	   - 0 means we haven't taken the lock within our transaction nest and need to acquire the lock.
-//     - Otherwise we already have the lock (and are preventing non-transactional code seeing any
-//       modifications we've made while holding the lock), so just bump `TransactionalLockCount`.
-//     - We also register an on-abort handler to release the lock should we abort (but we need to
-//       query `TransactionalLockCount` even there because we could be aborting an inner transaction
-//       and the parent transaction still wants to have the lock held!).
-//   - During unlocking we defer doing the unlock until the transaction commits.
-//
-// Thus with this approach we will hold this lock for the *entirety* of the transactional nest should
-// we take the lock during the transaction, thus preventing non-transactional code from seeing any
-// modifications we should make.
-struct FTransactionallySafeCriticalSection final
-{
-	FORCEINLINE void Lock()
-	{
-		if (AutoRTFM::IsTransactional())
-		{
-			AutoRTFM::Open([&]
-				{
-					// The transactional system which can increment TransactionalLockCount
-					// is always single-threaded, thus this is safe to check without atomicity.
-					if (0 == TransactionalLockCount)
-					{
-						CriticalSection.Lock();
-					}
-
-					TransactionalLockCount += 1;
-				});
-			
-			AutoRTFM::OnAbort([this]
-				{
-					check(0 != TransactionalLockCount);
-					TransactionalLockCount -= 1;
-
-					if (0 == TransactionalLockCount)
-					{
-						CriticalSection.Unlock();
-					}
-				});
-		}
-		else
-		{
-			CriticalSection.Lock();
-			check(0 == TransactionalLockCount);
-		}
-	}
-
-	void Unlock()
-	{
-		if (AutoRTFM::IsTransactional())
-		{
-			AutoRTFM::OnCommit([this]
-				{
-					check(0 != TransactionalLockCount);
-					TransactionalLockCount -= 1;
-
-					if (0 == TransactionalLockCount)
-					{
-						CriticalSection.Unlock();
-					}
-				});
-		}
-		else
-		{
-			check(0 == TransactionalLockCount);
-			CriticalSection.Unlock();
-		}
-	}
-
-private:
-	FCriticalSection CriticalSection;
-	uint32 TransactionalLockCount = 0;
-};
-
-struct FTransactionallySafeScopeLock final
-{
-	UE_NODISCARD_CTOR FTransactionallySafeScopeLock(FTransactionallySafeCriticalSection* InSynchObject)
-		: SynchObject(InSynchObject)
-	{
-		check(SynchObject);
-		SynchObject->Lock();
-	}
-
-	~FTransactionallySafeScopeLock()
-	{
-		Unlock();
-	}
-
-	void Unlock()
-	{
-		if (SynchObject)
-		{
-			SynchObject->Unlock();
-			SynchObject = nullptr;
-		}
-	}
-
-private:
-	FTransactionallySafeScopeLock() = delete;
-	FTransactionallySafeScopeLock(const FTransactionallySafeScopeLock&) = delete;
-
-	FTransactionallySafeCriticalSection* SynchObject;
-};
+#include "Misc/TransactionallySafeScopeLock.h"
 
 struct ENGINE_API FTraceFilterObjectAnnotation
 {
diff --git a/Engine/Source/Runtime/GameplayTags/Classes/GameplayTagsManager.h b/Engine/Source/Runtime/GameplayTags/Classes/GameplayTagsManager.h
--- a/Engine/Source/Runtime/GameplayTags/Classes/GameplayTagsManager.h
+++ b/Engine/Source/Runtime/GameplayTags/Classes/GameplayTagsManager.h
@@ -11,6 +11,7 @@
 #include "GameplayTagContainer.h"
 #include "Engine/DataTable.h"
 #include "Templates/UniquePtr.h"
+#include "Misc/TransactionallySafeScopeLock.h"
 
 #include "GameplayTagsManager.generated.h"
 
@@ -439,7 +440,7 @@ public:
 	 */
 	FORCEINLINE_DEBUGGABLE TSharedPtr<FGameplayTagNode> FindTagNode(const FGameplayTag& GameplayTag) const
 	{
-		FScopeLock Lock(&GameplayTagMapCritical);
+		FTransactionallySafeScopeLock Lock(&GameplayTagMapCritical);
 
 		const TSharedPtr<FGameplayTagNode>* Node = GameplayTagNodeMap.Find(GameplayTag);
 
@@ -927,7 +928,7 @@ private:
 
 	// This critical section is to handle an issue where tag requests come from another thread when async loading from a background thread in FGameplayTagContainer::Serialize.
 	// This class is not generically threadsafe.
-	mutable FCriticalSection GameplayTagMapCritical;
+	mutable FTransactionallySafeCriticalSection GameplayTagMapCritical;
 
 #if WITH_EDITOR
 	// Transient editor-only tags to support quick-iteration PIE workflows
diff --git a/Engine/Source/Runtime/GameplayTags/Private/GameplayTagsManager.cpp b/Engine/Source/Runtime/GameplayTags/Private/GameplayTagsManager.cpp
--- a/Engine/Source/Runtime/GameplayTags/Private/GameplayTagsManager.cpp
+++ b/Engine/Source/Runtime/GameplayTags/Private/GameplayTagsManager.cpp
@@ -570,7 +570,7 @@ int32 PrintNetIndiceAssignment = 0;
 static FAutoConsoleVariableRef CVarPrintNetIndiceAssignment(TEXT("GameplayTags.PrintNetIndiceAssignment"), PrintNetIndiceAssignment, TEXT("Logs GameplayTag NetIndice assignment"), ECVF_Default );
 void UGameplayTagsManager::ConstructNetIndex()
 {
-	FScopeLock Lock(&GameplayTagMapCritical);
+	FTransactionallySafeScopeLock Lock(&GameplayTagMapCritical);
 
 	bNetworkIndexInvalidated = false;
 
@@ -963,7 +963,7 @@ bool UGameplayTagsManager::ImportSingleGameplayTag(FGameplayTag& Tag, FName Impo
 	return bRetVal;
 }
 
-void UGameplayTagsManager::InitializeManager()
+UE_AUTORTFM_ALWAYS_OPEN void UGameplayTagsManager::InitializeManager()
 {
 	check(!SingletonManager);
 	SCOPED_BOOT_TIMING("UGameplayTagsManager::InitializeManager");
@@ -1205,7 +1205,7 @@ UGameplayTagsManager::~UGameplayTagsManager()
 
 void UGameplayTagsManager::DestroyGameplayTagTree()
 {
-	FScopeLock Lock(&GameplayTagMapCritical);
+	FTransactionallySafeScopeLock Lock(&GameplayTagMapCritical);
 
 	if (GameplayRootTag.IsValid())
 	{
@@ -1291,7 +1291,7 @@ int32 UGameplayTagsManager::InsertTagIntoNodeArray(FName Tag, FName FullTag, TSh
 		{
 			// This critical section is to handle an issue where tag requests come from another thread when async loading from a background thread in FGameplayTagContainer::Serialize.
 			// This function is not generically threadsafe.
-			FScopeLock Lock(&GameplayTagMapCritical);
+			FTransactionallySafeScopeLock Lock(&GameplayTagMapCritical);
 			GameplayTagNodeMap.Add(GameplayTag, TagNode);
 		}
 	}
@@ -1315,7 +1315,7 @@ void UGameplayTagsManager::PrintReplicationIndices()
 
 	UE_LOG(LogGameplayTags, Display, TEXT("::PrintReplicationIndices (TOTAL %d)"), GameplayTagNodeMap.Num());
 
-	FScopeLock Lock(&GameplayTagMapCritical);
+	FTransactionallySafeScopeLock Lock(&GameplayTagMapCritical);
 
 	for (auto It : GameplayTagNodeMap)
 	{
@@ -1636,7 +1636,7 @@ FString UGameplayTagsManager::GetCategoriesMetaFromFunction(const UFunction* Thi
 
 void UGameplayTagsManager::GetAllTagsFromSource(FName TagSource, TArray< TSharedPtr<FGameplayTagNode> >& OutTagArray) const
 {
-	FScopeLock Lock(&GameplayTagMapCritical);
+	FTransactionallySafeScopeLock Lock(&GameplayTagMapCritical);
 
 	for (const TPair<FGameplayTag, TSharedPtr<FGameplayTagNode>>& NodePair : GameplayTagNodeMap)
 	{
@@ -1958,7 +1958,7 @@ FGameplayTag UGameplayTagsManager::RequestGameplayTag(FName TagName, bool ErrorI
 
 	// This critical section is to handle an issue where tag requests come from another thread when async loading from a background thread in FGameplayTagContainer::Serialize.
 	// This function is not generically threadsafe.
-	FScopeLock Lock(&GameplayTagMapCritical);
+	FTransactionallySafeScopeLock Lock(&GameplayTagMapCritical);
 
 	// Check if there are redirects for this tag. If so and the redirected tag is in the node map, return it.
 	// Redirects take priority, even if the tag itself may exist.
@@ -2081,7 +2081,7 @@ FGameplayTag UGameplayTagsManager::FindGameplayTagFromPartialString_Slow(FString
 {
 	// This critical section is to handle an issue where tag requests come from another thread when async loading from a background thread in FGameplayTagContainer::Serialize.
 	// This function is not generically threadsafe.
-	FScopeLock Lock(&GameplayTagMapCritical);
+	FTransactionallySafeScopeLock Lock(&GameplayTagMapCritical);
 
 	// Exact match first
 	FGameplayTag PossibleTag(*PartialString);
@@ -2205,10 +2205,9 @@ void UGameplayTagsManager::DoneAddingNativeTags()
 	}
 }
 
-UE_AUTORTFM_ALWAYS_OPEN
 FGameplayTagContainer UGameplayTagsManager::RequestGameplayTagParents(const FGameplayTag& GameplayTag) const
 {
-	FScopeLock Lock(&GameplayTagMapCritical);
+	FTransactionallySafeScopeLock Lock(&GameplayTagMapCritical);
 
 	const FGameplayTagContainer* ParentTags = GetSingleTagContainerPtr(GameplayTag);
 
@@ -2240,41 +2239,37 @@ bool UGameplayTagsManager::ExtractParentTags(const FGameplayTag& GameplayTag, TA
 	int32 OldSize = UniqueParentTags.Num();
 	FName RawTag = GameplayTag.GetTagName();
 
-	// Need to run in the open as it takes a lock. 
-	UE_AUTORTFM_OPEN2
-		{
-			FScopeLock Lock(&GameplayTagMapCritical);
+	FTransactionallySafeScopeLock Lock(&GameplayTagMapCritical);
 
-			// This code does not check redirectors because that was already handled by GameplayTagContainerLoaded
-			const TSharedPtr<FGameplayTagNode>*Node = GameplayTagNodeMap.Find(GameplayTag);
-			if (Node)
-			{
-				// Use the registered tag container if it exists
-				const FGameplayTagContainer& SingleContainer = (*Node)->GetSingleTagContainer();
-				for (const FGameplayTag& ParentTag : SingleContainer.ParentTags)
-				{
-					UniqueParentTags.AddUnique(ParentTag);
-				}
+	// This code does not check redirectors because that was already handled by GameplayTagContainerLoaded
+	const TSharedPtr<FGameplayTagNode>*Node = GameplayTagNodeMap.Find(GameplayTag);
+	if (Node)
+	{
+		// Use the registered tag container if it exists
+		const FGameplayTagContainer& SingleContainer = (*Node)->GetSingleTagContainer();
+		for (const FGameplayTag& ParentTag : SingleContainer.ParentTags)
+		{
+			UniqueParentTags.AddUnique(ParentTag);
+		}
 
-				if constexpr (0 != VALIDATE_EXTRACT_PARENT_TAGS)
-				{
-					GameplayTag.ParseParentTags(ValidationCopy);
-					ensureAlwaysMsgf(ValidationCopy == UniqueParentTags, TEXT("ExtractParentTags results are inconsistent for tag %s"), *GameplayTag.ToString());
-				}
-			}
-			else
-			{
-				// If we don't clear invalid tags, we need to extract the parents now in case they get registered later
-				GameplayTag.ParseParentTags(UniqueParentTags);
-			}
-		};
+		if constexpr (0 != VALIDATE_EXTRACT_PARENT_TAGS)
+		{
+			GameplayTag.ParseParentTags(ValidationCopy);
+			ensureAlwaysMsgf(ValidationCopy == UniqueParentTags, TEXT("ExtractParentTags results are inconsistent for tag %s"), *GameplayTag.ToString());
+		}
+	}
+	else
+	{
+		// If we don't clear invalid tags, we need to extract the parents now in case they get registered later
+		GameplayTag.ParseParentTags(UniqueParentTags);
+	}
 
 	return UniqueParentTags.Num() != OldSize;
 }
 
 void UGameplayTagsManager::RequestAllGameplayTags(FGameplayTagContainer& TagContainer, bool OnlyIncludeDictionaryTags) const
 {
-	FScopeLock Lock(&GameplayTagMapCritical);
+	FTransactionallySafeScopeLock Lock(&GameplayTagMapCritical);
 
 	for (const TPair<FGameplayTag, TSharedPtr<FGameplayTagNode>>& NodePair : GameplayTagNodeMap)
 	{
