Fix GameplayTagsManager to ensure with useful message instead of hard crash if there are too many gameplay tags for replication

diff --git a/Engine/Source/Runtime/GameplayTags/Private/GameplayTagsManager.cpp b/Engine/Source/Runtime/GameplayTags/Private/GameplayTagsManager.cpp
--- a/Engine/Source/Runtime/GameplayTags/Private/GameplayTagsManager.cpp
+++ b/Engine/Source/Runtime/GameplayTags/Private/GameplayTagsManager.cpp
@@ -603,20 +603,20 @@ void UGameplayTagsManager::ConstructNetIndex()
 		checkf( Found, TEXT("Tag %s not found in NetworkGameplayTagNodeIndex"), *Tag.ToString() );
 	}
 
-	InvalidTagNetIndex = IntCastChecked<uint16, int32>(NetworkGameplayTagNodeIndex.Num() + 1);
-	NetIndexTrueBitNum = FMath::CeilToInt(FMath::Log2(static_cast<float>(InvalidTagNetIndex)));
-	
-	// This should never be smaller than NetIndexTrueBitNum
-	NetIndexFirstBitSegment = FMath::Min<int32>(GetDefault<UGameplayTagsSettings>()->NetIndexFirstBitSegment, NetIndexTrueBitNum);
-
 	// This is now sorted and it should be the same on both client and server
 	if (NetworkGameplayTagNodeIndex.Num() >= INVALID_TAGNETINDEX)
 	{
-		ensureMsgf(false, TEXT("Too many tags in dictionary for networking! Remove tags or increase tag net index size"));
+		ensureMsgf(false, TEXT("Too many tags (%d) in dictionary for networking! Remove tags or increase tag net index size (%d)"), NetworkGameplayTagNodeIndex.Num(), INVALID_TAGNETINDEX);
 
 		NetworkGameplayTagNodeIndex.SetNum(INVALID_TAGNETINDEX - 1);
 	}
 
+	InvalidTagNetIndex = IntCastChecked<uint16, int32>(NetworkGameplayTagNodeIndex.Num() + 1);
+	NetIndexTrueBitNum = FMath::CeilToInt(FMath::Log2(static_cast<float>(InvalidTagNetIndex)));
+
+	// This should never be smaller than NetIndexTrueBitNum
+	NetIndexFirstBitSegment = FMath::Min<int32>(GetDefault<UGameplayTagsSettings>()->NetIndexFirstBitSegment, NetIndexTrueBitNum);
+
 	UE_CLOG(PrintNetIndiceAssignment, LogGameplayTags, Display, TEXT("Assigning NetIndices to %d tags."), NetworkGameplayTagNodeIndex.Num() );
 
 	NetworkGameplayTagNodeIndexHash = 0;
