#UE Add Simple accessor to see how many tags exist for UI reasons

diff --git a/Engine/Source/Runtime/GameplayTags/Classes/GameplayTagsManager.h b/Engine/Source/Runtime/GameplayTags/Classes/GameplayTagsManager.h
--- a/Engine/Source/Runtime/GameplayTags/Classes/GameplayTagsManager.h
+++ b/Engine/Source/Runtime/GameplayTags/Classes/GameplayTagsManager.h
@@ -748,6 +748,7 @@ public:
 #endif //WITH_EDITOR
 
 	GAMEPLAYTAGS_API void PrintReplicationIndices();
+	GAMEPLAYTAGS_API int32 GetNumGameplayTagNodes() const { return GameplayTagNodeMap.Num(); }
 
 #if !(UE_BUILD_SHIPPING || UE_BUILD_TEST)
 	/** Mechanism for tracking what tags are frequently replicated */
