Console command to dump gameplay tag sources.
Fix not registering a source for some native tags created during static initialization before the gameplay tag manager is available.

diff --git a/Engine/Source/Runtime/GameplayTags/Classes/GameplayTagsManager.h b/Engine/Source/Runtime/GameplayTags/Classes/GameplayTagsManager.h
--- a/Engine/Source/Runtime/GameplayTags/Classes/GameplayTagsManager.h
+++ b/Engine/Source/Runtime/GameplayTags/Classes/GameplayTagsManager.h
@@ -678,7 +678,6 @@ public:
 	/** Returns information about tag. If not found return false */
     GAMEPLAYTAGS_API bool GetTagEditorData(FName TagName, FString& OutComment, TArray<FName>& OutTagSources, bool& bOutIsTagExplicit, bool &bOutIsRestrictedTag, bool &bOutAllowNonRestrictedChildren) const;
 
-#if WITH_EDITOR
 	/** This is called after EditorRefreshGameplayTagTree. Useful if you need to do anything editor related when tags are added or removed */
 	static GAMEPLAYTAGS_API FSimpleMulticastDelegate OnEditorRefreshGameplayTagTree;
 
@@ -690,7 +689,6 @@ public:
 
 	/** Resumes EditorRefreshGameplayTagTree requests; triggers a refresh if a request was made while it was suspended */
 	GAMEPLAYTAGS_API void ResumeEditorRefreshGameplayTagTree(FGuid SuspendToken);
-#endif //if WITH_EDITOR
 
 	/** Gets a Tag Container containing all of the tags in the hierarchy that are children of this tag, and were explicitly added to the dictionary */
 	GAMEPLAYTAGS_API FGameplayTagContainer RequestGameplayTagChildrenInDictionary(const FGameplayTag& GameplayTag) const;
@@ -744,7 +742,8 @@ public:
 	
 	GAMEPLAYTAGS_API bool ShowGameplayTagAsHyperLinkEditor(FString TagName);
 
-
+	/** Implementation of console command GameplayTags.DumpSources */
+	void DumpSources(FOutputDevice& Out) const;
 #endif //WITH_EDITOR
 
 	GAMEPLAYTAGS_API void PrintReplicationIndices();
diff --git a/Engine/Source/Runtime/GameplayTags/Private/GameplayTagsManager.cpp b/Engine/Source/Runtime/GameplayTags/Private/GameplayTagsManager.cpp
--- a/Engine/Source/Runtime/GameplayTags/Private/GameplayTagsManager.cpp
+++ b/Engine/Source/Runtime/GameplayTags/Private/GameplayTagsManager.cpp
@@ -75,6 +75,16 @@ static FAutoConsoleCommand PrintReplicationFrequencyReportCommand(
 
 #endif
 
+#if WITH_EDITOR
+static FAutoConsoleCommand CMD_DumpGameplayTagSources(
+	TEXT("GameplayTags.DumpSources"),
+	TEXT("Dumps all known sources of gameplay tags"),
+	FConsoleCommandWithOutputDeviceDelegate::CreateLambda([](FOutputDevice& Out){
+		UGameplayTagsManager::Get().DumpSources(Out);
+	})
+);
+#endif
+
 struct FCompareFGameplayTagNodeByTag
 {
 	FORCEINLINE bool operator()(const TSharedPtr<FGameplayTagNode>& A, const TSharedPtr<FGameplayTagNode>& B) const
@@ -446,6 +456,7 @@ void UGameplayTagsManager::ConstructGameplayTagTree()
 
 			for (const class FNativeGameplayTag* NativeTag : FNativeGameplayTag::GetRegisteredNativeTags())
 			{
+				FindOrAddTagSource(NativeTag->GetModuleName(), EGameplayTagSourceType::Native);
 				AddTagTableRow(NativeTag->GetGameplayTagTableRow(), NativeTag->GetModuleName());
 			}
 		}
@@ -2424,6 +2435,19 @@ bool UGameplayTagsManager::ValidateTagCreation(FName TagName) const
 	return FindTagNode(TagName).IsValid();
 }
 
+void UGameplayTagsManager::DumpSources(FOutputDevice& Out) const
+{
+	for (const TPair<FName, FGameplayTagSource>& Pair : TagSources)
+	{
+		Out.Logf(TEXT("%s : %s"), *Pair.Key.ToString(), *UEnum::GetValueAsString(Pair.Value.SourceType));
+		FString ConfigFilePath = Pair.Value.GetConfigFileName();
+		if (!ConfigFilePath.IsEmpty())
+		{
+			Out.Logf(TEXT("Config file path: %s"), *Pair.Value.SourceTagList->ConfigFileName);
+		}
+	}
+}
+
 FGameplayTagTableRow::FGameplayTagTableRow(FGameplayTagTableRow const& Other)
 {
 	*this = Other;
