Fix definition of DumpSources to be in WITH_EDITOR.

diff --git a/Engine/Source/Runtime/GameplayTags/Private/GameplayTagsManager.cpp b/Engine/Source/Runtime/GameplayTags/Private/GameplayTagsManager.cpp
--- a/Engine/Source/Runtime/GameplayTags/Private/GameplayTagsManager.cpp
+++ b/Engine/Source/Runtime/GameplayTags/Private/GameplayTagsManager.cpp
@@ -2435,6 +2435,7 @@ bool UGameplayTagsManager::ValidateTagCreation(FName TagName) const
 	return FindTagNode(TagName).IsValid();
 }
 
+#if WITH_EDITOR
 void UGameplayTagsManager::DumpSources(FOutputDevice& Out) const
 {
 	for (const TPair<FName, FGameplayTagSource>& Pair : TagSources)
@@ -2447,6 +2448,7 @@ void UGameplayTagsManager::DumpSources(FOutputDevice& Out) const
 		}
 	}
 }
+#endif
 
 FGameplayTagTableRow::FGameplayTagTableRow(FGameplayTagTableRow const& Other)
 {
