#UE Fix bug where UGameplayTagsManager::AddChildrenTags would add non-dictionary tags when you asked for dicationary tags

diff --git a/Engine/Source/Runtime/GameplayTags/Private/GameplayTagsManager.cpp b/Engine/Source/Runtime/GameplayTags/Private/GameplayTagsManager.cpp
--- a/Engine/Source/Runtime/GameplayTags/Private/GameplayTagsManager.cpp
+++ b/Engine/Source/Runtime/GameplayTags/Private/GameplayTagsManager.cpp
@@ -2330,7 +2330,7 @@ void UGameplayTagsManager::AddChildrenTags(FGameplayTagContainer& TagContainer,
 				bool bShouldInclude = true;
 
 #if WITH_EDITORONLY_DATA
-				if (OnlyIncludeDictionaryTags && ChildNode->GetFirstSourceName() == NAME_None)
+				if (OnlyIncludeDictionaryTags && !ChildNode->IsExplicitTag())
 				{
 					// Only have info to do this in editor builds
 					bShouldInclude = false;
