Modified version of the fix found in PR #12052: [????] GetOwnedGameplayTags Fix For Reusing old tag container in blueprint event graphs.

The (Blueprint) GameplayTagInterface::GetOwnedGameplayTags member function has been redirected to BlueprintGameplayTagLibrary::GetOwnedGameplayTags.  The implementation ensures that the results you get are always a newly initialized tag container, thus ensuring you do not accidentally reuse and append old results.

diff --git a/Engine/Config/BaseEngine.ini b/Engine/Config/BaseEngine.ini
--- a/Engine/Config/BaseEngine.ini
+++ b/Engine/Config/BaseEngine.ini
@@ -1302,6 +1302,11 @@ DefaultWorldPartitionHLODLayer=/Engine/Maps/Templates/HLODs/HLODLayer_Instanced.
 
 ; End of releases redirects
 
+; Fixing the output of tags for blueprints to be reset.
++FunctionRedirects=(OldName="/Script/GameplayTags.GameplayTagAssetInterface.GetOwnedGameplayTags", NewName="BlueprintGameplayTagLibrary.GetOwnedGameplayTags")
++PropertyRedirects=(OldName="BlueprintGameplayTagLibrary.GetOwnedGameplayTags.TagContainer", NewName="ReturnValue")
++PropertyRedirects=(OldName="BlueprintGameplayTagLibrary.GetOwnedGameplayTags.Self", NewName="TagContainerInterface")
+
 [CoreUObject.Metadata]
 ; Note: UnrealHeaderTool should be rerun after making changes to MetadataRedirects to catch any keys specified in class headers
 +MetadataRedirects=(OldKey="K2Protected", NewKey="BlueprintProtected")
diff --git a/Engine/Source/Runtime/GameplayTags/Classes/BlueprintGameplayTagLibrary.h b/Engine/Source/Runtime/GameplayTags/Classes/BlueprintGameplayTagLibrary.h
--- a/Engine/Source/Runtime/GameplayTags/Classes/BlueprintGameplayTagLibrary.h
+++ b/Engine/Source/Runtime/GameplayTags/Classes/BlueprintGameplayTagLibrary.h
@@ -257,7 +257,15 @@ class UBlueprintGameplayTagLibrary : public UBlueprintFunctionLibrary
 	/** Checks if a gameplay tag containers's name and a string are not equal to one another */
 	UFUNCTION(BlueprintPure, Category = PinOptions, meta = (BlueprintInternalUseOnly = "TRUE"))
 	static GAMEPLAYTAGS_API bool NotEqual_TagContainerTagContainer(FGameplayTagContainer A, FString B);
-	
+
+	/** @return the tags currently owned by the TagContainerInterface object */
+	UFUNCTION(BlueprintPure, Category = "GameplayTags", meta=(DefaultToSelf=TagContainerInterface))
+	static GAMEPLAYTAGS_API UPARAM(DisplayName = "Owned Tags") FGameplayTagContainer GetOwnedGameplayTags(TScriptInterface<IGameplayTagAssetInterface> TagContainerInterface);
+
+	/** Converts a UObject to a GameplayTagAssetInterface. This specialty Autocast function exists so we can auto-convert the GameplayTagAssetInterface member functions to static Blueprint functions using redirects. */ 
+	UFUNCTION(BlueprintPure, meta=(DisplayName = "To GameplayTagAssetInterface", CompactNodeTitle = "->", BlueprintAutocast), Category="Utilities|Casting")
+	static GAMEPLAYTAGS_API TScriptInterface<IGameplayTagAssetInterface> Conv_ObjectToGameplayTagAssetInterface(UObject* InObject);
+
 	/**
 	 * Returns an FString listing all of the gameplay tags in the tag container for debugging purposes.
 	 *
diff --git a/Engine/Source/Runtime/GameplayTags/Classes/GameplayTagAssetInterface.h b/Engine/Source/Runtime/GameplayTags/Classes/GameplayTagAssetInterface.h
--- a/Engine/Source/Runtime/GameplayTags/Classes/GameplayTagAssetInterface.h
+++ b/Engine/Source/Runtime/GameplayTags/Classes/GameplayTagAssetInterface.h
@@ -25,7 +25,6 @@ class IGameplayTagAssetInterface
 	 * 
 	 * @param OutTags	[OUT] Set of tags on the asset
 	 */
-	 UFUNCTION(BlueprintCallable, Category = GameplayTags)
 	virtual void GetOwnedGameplayTags(FGameplayTagContainer& TagContainer) const=0;
 
 	/**
diff --git a/Engine/Source/Runtime/GameplayTags/Private/BlueprintGameplayTagLibrary.cpp b/Engine/Source/Runtime/GameplayTags/Private/BlueprintGameplayTagLibrary.cpp
--- a/Engine/Source/Runtime/GameplayTags/Private/BlueprintGameplayTagLibrary.cpp
+++ b/Engine/Source/Runtime/GameplayTags/Private/BlueprintGameplayTagLibrary.cpp
@@ -298,6 +298,30 @@ bool UBlueprintGameplayTagLibrary::NotEqual_TagContainerTagContainer(FGameplayTa
 
 	return A != TagContainer;
 }
+
+FGameplayTagContainer UBlueprintGameplayTagLibrary::GetOwnedGameplayTags(TScriptInterface<IGameplayTagAssetInterface> TagContainerInterface)
+{
+	// Note:
+	// Why do this instead of just letting the BP Compiler use the InOut parameter of the original function?
+	//   Because the BP Compiler will store the out values in the event graph as a hidden property, which will
+	//   cause issues because the tags will never get reset, and will always be appended to, most all the implementations
+	//   of the GetOwnedGameplayTags allow you to continually Append tags to the container, but because the fact that the
+	//   inout is hidden in BP, you don't get the understand or option really to clear the tags, so this is a safer way.
+
+	FGameplayTagContainer OwnedTags;
+	if (TagContainerInterface)
+	{
+		TagContainerInterface->GetOwnedGameplayTags(OwnedTags);
+	}
+
+	return OwnedTags;
+}
+
+TScriptInterface<IGameplayTagAssetInterface> UBlueprintGameplayTagLibrary::Conv_ObjectToGameplayTagAssetInterface(UObject* InObject)
+{
+	return TScriptInterface<IGameplayTagAssetInterface>(InObject);
+}
+
 FString UBlueprintGameplayTagLibrary::GetDebugStringFromGameplayTagContainer(const FGameplayTagContainer& TagContainer)
 {
 	return TagContainer.ToStringSimple();
