PR #12052: [????] GetOwnedGameplayTags Fix For Reusing old tag container in blueprint event graphs

diff --git a/Engine/Config/BaseEngine.ini b/Engine/Config/BaseEngine.ini
--- a/Engine/Config/BaseEngine.ini
+++ b/Engine/Config/BaseEngine.ini
@@ -1329,6 +1329,10 @@ DefaultWorldPartitionHLODLayer=/Engine/Maps/Templates/HLODs/HLODLayer_Instanced.
 
 ; End of releases redirects
 
+; Fixing the output of tags for blueprints to be reset.
++FunctionRedirects=(OldName="/Script/GameplayTags.GameplayTagAssetInterface.GetOwnedGameplayTags", NewName="BP_GetOwnedGameplayTags")
++PropertyRedirects=(OldName="GameplayTagAssetInterface.BP_GetOwnedGameplayTags.TagContainer", NewName="ReturnValue")
+
 [CoreUObject.Metadata]
 ; Note: UnrealHeaderTool should be rerun after making changes to MetadataRedirects to catch any keys specified in class headers
 +MetadataRedirects=(OldKey="K2Protected", NewKey="BlueprintProtected")
diff --git a/Engine/Source/Runtime/GameplayTags/Classes/GameplayTagAssetInterface.h b/Engine/Source/Runtime/GameplayTags/Classes/GameplayTagAssetInterface.h
--- a/Engine/Source/Runtime/GameplayTags/Classes/GameplayTagAssetInterface.h
+++ b/Engine/Source/Runtime/GameplayTags/Classes/GameplayTagAssetInterface.h
@@ -25,7 +25,6 @@ class IGameplayTagAssetInterface
 	 * 
 	 * @param OutTags	[OUT] Set of tags on the asset
 	 */
-	 UFUNCTION(BlueprintCallable, Category = GameplayTags)
 	virtual void GetOwnedGameplayTags(FGameplayTagContainer& TagContainer) const=0;
 
 	/**
@@ -57,5 +56,12 @@ class IGameplayTagAssetInterface
 	 */
 	UFUNCTION(BlueprintCallable, Category=GameplayTags)
 	GAMEPLAYTAGS_API virtual bool HasAnyMatchingGameplayTags(const FGameplayTagContainer& TagContainer) const;
-};
 
+protected:
+
+	/**
+	 * Gets the owned gameplay tags for the asset.  Exposed to allow redirects of existing GetOwnedGameplayTags calls.  In Blueprints, new nodes will use BlueprintGameplayTagLibrary's version.
+	 */
+	UFUNCTION(BlueprintCallable, Category = GameplayTags, BlueprintInternalUseOnly, meta=(DisplayName="Get Owned Gameplay Tags", AllowPrivateAccess=true))
+	GAMEPLAYTAGS_API virtual UPARAM(DisplayName = "Owned Tags") FGameplayTagContainer BP_GetOwnedGameplayTags() const;
+};
diff --git a/Engine/Source/Runtime/GameplayTags/Private/GameplayTagAssetInterface.cpp b/Engine/Source/Runtime/GameplayTags/Private/GameplayTagAssetInterface.cpp
--- a/Engine/Source/Runtime/GameplayTags/Private/GameplayTagAssetInterface.cpp
+++ b/Engine/Source/Runtime/GameplayTags/Private/GameplayTagAssetInterface.cpp
@@ -1,6 +1,7 @@
 // Copyright Epic Games, Inc. All Rights Reserved.
 
 #include "GameplayTagAssetInterface.h"
+#include "BlueprintGameplayTagLibrary.h"
 
 #include UE_INLINE_GENERATED_CPP_BY_NAME(GameplayTagAssetInterface)
 
@@ -33,3 +34,14 @@ bool IGameplayTagAssetInterface::HasAnyMatchingGameplayTags(const FGameplayTagCo
 	return OwnedTags.HasAny(TagContainer);
 }
 
+FGameplayTagContainer IGameplayTagAssetInterface::BP_GetOwnedGameplayTags() const
+{
+	// This is just a hack to forward to the proper non-member function.  The idea being in the future, 
+	// we can redirect this member function to the non-member function in the future once we allow the CoreRedirect
+	// of the 'self' pin.
+	TScriptInterface<IGameplayTagAssetInterface> GameplayTagAssetInterface;
+	GameplayTagAssetInterface.SetObject(_getUObject());
+	GameplayTagAssetInterface.SetInterface(const_cast<IGameplayTagAssetInterface*>(this));
+
+	return UBlueprintGameplayTagLibrary::GetOwnedGameplayTags(GameplayTagAssetInterface);
+}
