Exposing the Filter function from GameplayTagContainer to Blueprint.

#rnx

diff --git a/Engine/Source/Runtime/GameplayTags/Classes/BlueprintGameplayTagLibrary.h b/Engine/Source/Runtime/GameplayTags/Classes/BlueprintGameplayTagLibrary.h
--- a/Engine/Source/Runtime/GameplayTags/Classes/BlueprintGameplayTagLibrary.h
+++ b/Engine/Source/Runtime/GameplayTags/Classes/BlueprintGameplayTagLibrary.h
@@ -106,6 +106,17 @@ class UBlueprintGameplayTagLibrary : public UBlueprintFunctionLibrary
 	UFUNCTION(BlueprintPure, Category="GameplayTags", meta = (Keywords = "DoesContainerMatchAllTagsInContainer", BlueprintThreadSafe))
 	static GAMEPLAYTAGS_API bool HasAllTags(const FGameplayTagContainer& TagContainer, const FGameplayTagContainer& OtherContainer, bool bExactMatch);
 
+	/**
+	 * Returns a filtered version of this container, returns all tags that match against any of the tags in OtherContainer
+	 *
+	 * @param OtherContainer		The Container to filter against
+	 * @param bExactMatch			If true, the tag has to be exactly present, if false then TagContainer will include it's parent tags while matching
+	 *
+	 * @return A FGameplayTagContainer containing the filtered tags
+	 */
+	UFUNCTION(BlueprintPure, Category="GameplayTags", meta = (BlueprintThreadSafe))
+	static GAMEPLAYTAGS_API FGameplayTagContainer Filter(const FGameplayTagContainer& TagContainer, const FGameplayTagContainer& OtherContainer, bool bExactMatch);
+
 	/**
 	 * Check if the specified tag query is empty
 	 * 
diff --git a/Engine/Source/Runtime/GameplayTags/Private/BlueprintGameplayTagLibrary.cpp b/Engine/Source/Runtime/GameplayTags/Private/BlueprintGameplayTagLibrary.cpp
--- a/Engine/Source/Runtime/GameplayTags/Private/BlueprintGameplayTagLibrary.cpp
+++ b/Engine/Source/Runtime/GameplayTags/Private/BlueprintGameplayTagLibrary.cpp
@@ -88,6 +88,15 @@ bool UBlueprintGameplayTagLibrary::HasAllTags(const FGameplayTagContainer& TagCo
 	return TagContainer.HasAll(OtherContainer);
 }
 
+FGameplayTagContainer UBlueprintGameplayTagLibrary::Filter(const FGameplayTagContainer& TagContainer, const FGameplayTagContainer& OtherContainer, bool bExactMatch)
+{
+	if (bExactMatch)
+	{
+		return TagContainer.FilterExact(OtherContainer);
+	}
+	return TagContainer.Filter(OtherContainer);
+}
+
 bool UBlueprintGameplayTagLibrary::IsTagQueryEmpty(const FGameplayTagQuery& TagQuery)
 {
 	return TagQuery.IsEmpty();
