Disabling GameplayTags Dynamic Replication by default because it requires IRIS to compile.

diff --git a/Engine/Source/Runtime/GameplayTags/Private/GameplayTagsSettings.cpp b/Engine/Source/Runtime/GameplayTags/Private/GameplayTagsSettings.cpp
--- a/Engine/Source/Runtime/GameplayTags/Private/GameplayTagsSettings.cpp
+++ b/Engine/Source/Runtime/GameplayTags/Private/GameplayTagsSettings.cpp
@@ -64,7 +64,7 @@ UGameplayTagsSettings::UGameplayTagsSettings(const FObjectInitializer& ObjectIni
 	ImportTagsFromConfig = true;
 	WarnOnInvalidTags = true;
 	FastReplication = false;
-	bDynamicReplication = true;
+	bDynamicReplication = false; // This currently needs Iris enabled (so don't enable by default)
 	AllowEditorTagUnloading = true;
 	AllowGameTagUnloading = false;
 	InvalidTagCharacters = ("\"',");
