[AutoRTFM] Migrate to UE::TScopeLock instead of FTransactionallySafeScopeLock.

These are functionally identical.

diff --git a/Engine/Source/Runtime/GameplayTags/Classes/GameplayTagsManager.h b/Engine/Source/Runtime/GameplayTags/Classes/GameplayTagsManager.h
--- a/Engine/Source/Runtime/GameplayTags/Classes/GameplayTagsManager.h
+++ b/Engine/Source/Runtime/GameplayTags/Classes/GameplayTagsManager.h
@@ -11,7 +11,8 @@
 #include "GameplayTagContainer.h"
 #include "Engine/DataTable.h"
 #include "Templates/UniquePtr.h"
-#include "Misc/TransactionallySafeScopeLock.h"
+#include "Misc/ScopeLock.h"
+#include "Misc/TransactionallySafeCriticalSection.h"
 
 #include "GameplayTagsManager.generated.h"
 
@@ -440,7 +441,7 @@ public:
 	 */
 	FORCEINLINE_DEBUGGABLE TSharedPtr<FGameplayTagNode> FindTagNode(const FGameplayTag& GameplayTag) const
 	{
-		FTransactionallySafeScopeLock Lock(&GameplayTagMapCritical);
+		UE::TScopeLock Lock(GameplayTagMapCritical);
 
 		const TSharedPtr<FGameplayTagNode>* Node = GameplayTagNodeMap.Find(GameplayTag);
 
diff --git a/Engine/Source/Runtime/GameplayTags/Private/GameplayTagsManager.cpp b/Engine/Source/Runtime/GameplayTags/Private/GameplayTagsManager.cpp
--- a/Engine/Source/Runtime/GameplayTags/Private/GameplayTagsManager.cpp
+++ b/Engine/Source/Runtime/GameplayTags/Private/GameplayTagsManager.cpp
@@ -582,7 +582,7 @@ int32 PrintNetIndiceAssignment = 0;
 static FAutoConsoleVariableRef CVarPrintNetIndiceAssignment(TEXT("GameplayTags.PrintNetIndiceAssignment"), PrintNetIndiceAssignment, TEXT("Logs GameplayTag NetIndice assignment"), ECVF_Default );
 void UGameplayTagsManager::ConstructNetIndex()
 {
-	FTransactionallySafeScopeLock Lock(&GameplayTagMapCritical);
+	UE::TScopeLock Lock(GameplayTagMapCritical);
 
 	bNetworkIndexInvalidated = false;
 
@@ -1217,7 +1217,7 @@ UGameplayTagsManager::~UGameplayTagsManager()
 
 void UGameplayTagsManager::DestroyGameplayTagTree()
 {
-	FTransactionallySafeScopeLock Lock(&GameplayTagMapCritical);
+	UE::TScopeLock Lock(GameplayTagMapCritical);
 
 	if (GameplayRootTag.IsValid())
 	{
@@ -1303,7 +1303,7 @@ int32 UGameplayTagsManager::InsertTagIntoNodeArray(FName Tag, FName FullTag, TSh
 		{
 			// This critical section is to handle an issue where tag requests come from another thread when async loading from a background thread in FGameplayTagContainer::Serialize.
 			// This function is not generically threadsafe.
-			FTransactionallySafeScopeLock Lock(&GameplayTagMapCritical);
+			UE::TScopeLock Lock(GameplayTagMapCritical);
 			GameplayTagNodeMap.Add(GameplayTag, TagNode);
 		}
 	}
@@ -1327,7 +1327,7 @@ void UGameplayTagsManager::PrintReplicationIndices()
 
 	UE_LOG(LogGameplayTags, Display, TEXT("::PrintReplicationIndices (TOTAL %d)"), GameplayTagNodeMap.Num());
 
-	FTransactionallySafeScopeLock Lock(&GameplayTagMapCritical);
+	UE::TScopeLock Lock(GameplayTagMapCritical);
 
 	for (auto It : GameplayTagNodeMap)
 	{
@@ -1648,7 +1648,7 @@ FString UGameplayTagsManager::GetCategoriesMetaFromFunction(const UFunction* Thi
 
 void UGameplayTagsManager::GetAllTagsFromSource(FName TagSource, TArray< TSharedPtr<FGameplayTagNode> >& OutTagArray) const
 {
-	FTransactionallySafeScopeLock Lock(&GameplayTagMapCritical);
+	UE::TScopeLock Lock(GameplayTagMapCritical);
 
 	for (const TPair<FGameplayTag, TSharedPtr<FGameplayTagNode>>& NodePair : GameplayTagNodeMap)
 	{
@@ -1970,7 +1970,7 @@ FGameplayTag UGameplayTagsManager::RequestGameplayTag(FName TagName, bool ErrorI
 
 	// This critical section is to handle an issue where tag requests come from another thread when async loading from a background thread in FGameplayTagContainer::Serialize.
 	// This function is not generically threadsafe.
-	FTransactionallySafeScopeLock Lock(&GameplayTagMapCritical);
+	UE::TScopeLock Lock(GameplayTagMapCritical);
 
 	// Check if there are redirects for this tag. If so and the redirected tag is in the node map, return it.
 	// Redirects take priority, even if the tag itself may exist.
@@ -2093,7 +2093,7 @@ FGameplayTag UGameplayTagsManager::FindGameplayTagFromPartialString_Slow(FString
 {
 	// This critical section is to handle an issue where tag requests come from another thread when async loading from a background thread in FGameplayTagContainer::Serialize.
 	// This function is not generically threadsafe.
-	FTransactionallySafeScopeLock Lock(&GameplayTagMapCritical);
+	UE::TScopeLock Lock(GameplayTagMapCritical);
 
 	// Exact match first
 	FGameplayTag PossibleTag(*PartialString);
@@ -2219,7 +2219,7 @@ void UGameplayTagsManager::DoneAddingNativeTags()
 
 FGameplayTagContainer UGameplayTagsManager::RequestGameplayTagParents(const FGameplayTag& GameplayTag) const
 {
-	FTransactionallySafeScopeLock Lock(&GameplayTagMapCritical);
+	UE::TScopeLock Lock(GameplayTagMapCritical);
 
 	const FGameplayTagContainer* ParentTags = GetSingleTagContainerPtr(GameplayTag);
 
@@ -2251,7 +2251,7 @@ bool UGameplayTagsManager::ExtractParentTags(const FGameplayTag& GameplayTag, TA
 	int32 OldSize = UniqueParentTags.Num();
 	FName RawTag = GameplayTag.GetTagName();
 
-	FTransactionallySafeScopeLock Lock(&GameplayTagMapCritical);
+	UE::TScopeLock Lock(GameplayTagMapCritical);
 
 	// This code does not check redirectors because that was already handled by GameplayTagContainerLoaded
 	const TSharedPtr<FGameplayTagNode>*Node = GameplayTagNodeMap.Find(GameplayTag);
@@ -2281,7 +2281,7 @@ bool UGameplayTagsManager::ExtractParentTags(const FGameplayTag& GameplayTag, TA
 
 void UGameplayTagsManager::RequestAllGameplayTags(FGameplayTagContainer& TagContainer, bool OnlyIncludeDictionaryTags) const
 {
-	FTransactionallySafeScopeLock Lock(&GameplayTagMapCritical);
+	UE::TScopeLock Lock(GameplayTagMapCritical);
 
 	for (const TPair<FGameplayTag, TSharedPtr<FGameplayTagNode>>& NodePair : GameplayTagNodeMap)
 	{
