In Replay, Clyde Battlepass cosmetics color variants appear as default variants

diff --git a/Engine/Source/Runtime/GameplayTags/Private/GameplayTagContainer.cpp b/Engine/Source/Runtime/GameplayTags/Private/GameplayTagContainer.cpp
--- a/Engine/Source/Runtime/GameplayTags/Private/GameplayTagContainer.cpp
+++ b/Engine/Source/Runtime/GameplayTags/Private/GameplayTagContainer.cpp
@@ -1351,7 +1351,7 @@ namespace UE::GameplayTags::GameplayTagDynamicSerialization
 		else if (Ar.IsLoading())
 		{
 			// Read TagNetToken
-			TagTokenDataStore->ReadNetToken(Ar);
+			TagToken = TagTokenDataStore->ReadNetToken(Ar);
 
 			if (Ar.IsError())
 			{
