#UE TagList sources other than the default can now define tag redirects. When renaming a tag redirect, we now put the redirect in the old source's location

diff --git a/Engine/Plugins/Editor/GameplayTagsEditor/Source/GameplayTagsEditor/Private/GameplayTagsEditorModule.cpp b/Engine/Plugins/Editor/GameplayTagsEditor/Source/GameplayTagsEditor/Private/GameplayTagsEditorModule.cpp
--- a/Engine/Plugins/Editor/GameplayTagsEditor/Source/GameplayTagsEditor/Private/GameplayTagsEditorModule.cpp
+++ b/Engine/Plugins/Editor/GameplayTagsEditor/Source/GameplayTagsEditor/Private/GameplayTagsEditorModule.cpp
@@ -827,7 +827,6 @@ public:
 		FName NewTagName = FName(*TagToRenameTo);
 
 		UGameplayTagsManager& Manager = UGameplayTagsManager::Get();
-		UGameplayTagsSettings* Settings = GetMutableDefault<UGameplayTagsSettings>();
 
 		FString OldComment, NewComment;
 		FName OldTagSourceName, NewTagSourceName;
@@ -839,6 +838,7 @@ public:
 		DeleteTagRedirector(NewTagName);
 		DeleteTagRedirector(OldTagName);
 
+		const FGameplayTagSource* OldTagSource = nullptr;
 		if (Manager.GetTagEditorData(OldTagName, OldComment, OldTagSourceName, bTagIsExplicit, bTagIsRestricted, bTagAllowsNonRestrictedChildren))
 		{
 			// Add new tag if needed
@@ -852,7 +852,7 @@ public:
 			}
 
 			// Delete old tag if possible, still make redirector if this fails
-			const FGameplayTagSource* OldTagSource = Manager.FindTagSource(OldTagSourceName);
+			OldTagSource = Manager.FindTagSource(OldTagSourceName);
 
 			if (OldTagSource && OldTagSource->SourceTagList)
 			{
@@ -883,11 +883,14 @@ public:
 		Redirect.OldTagName = OldTagName;
 		Redirect.NewTagName = NewTagName;
 
-		Settings->GameplayTagRedirects.AddUnique(Redirect);
+		UGameplayTagsList* ListToUpdate = (OldTagSource && OldTagSource->SourceTagList) ? OldTagSource->SourceTagList : GetMutableDefault<UGameplayTagsSettings>();
+		check(ListToUpdate);
 
-		GameplayTagsUpdateSourceControl(Settings->GetDefaultConfigFilename());
-		Settings->TryUpdateDefaultConfigFile();
-		GConfig->LoadFile(Settings->GetDefaultConfigFilename());
+		ListToUpdate->GameplayTagRedirects.AddUnique(Redirect);
+
+		GameplayTagsUpdateSourceControl(ListToUpdate->ConfigFileName);
+		ListToUpdate->TryUpdateDefaultConfigFile(ListToUpdate->ConfigFileName);
+		GConfig->LoadFile(ListToUpdate->ConfigFileName);
 
 		ShowNotification(FText::Format(LOCTEXT("AddTagRedirect", "Renamed tag {0} to {1}"), FText::FromString(TagToRename), FText::FromString(TagToRenameTo)), 3.0f);
 
diff --git a/Engine/Source/Runtime/GameplayTags/Private/GameplayTagRedirectors.cpp b/Engine/Source/Runtime/GameplayTags/Private/GameplayTagRedirectors.cpp
--- a/Engine/Source/Runtime/GameplayTags/Private/GameplayTagRedirectors.cpp
+++ b/Engine/Source/Runtime/GameplayTags/Private/GameplayTagRedirectors.cpp
@@ -2,6 +2,7 @@
 
 #include "GameplayTagRedirectors.h"
 #include "GameplayTagsSettings.h"
+#include "GameplayTagsManager.h"
 #include "Misc/ConfigCacheIni.h"
 
 #include UE_INLINE_GENERATED_CPP_BY_NAME(GameplayTagRedirectors)
@@ -60,19 +61,53 @@ FGameplayTagRedirectors::FGameplayTagRedirectors()
 	RefreshTagRedirects();
 }
 
+void FGameplayTagRedirectors::AddRedirectsFromSource(const FGameplayTagSource* Source)
+{
+	if (Source && Source->SourceTagList)
+	{
+		AddRedirects(Source->SourceTagList->GameplayTagRedirects);
+	}
+}
+
 void FGameplayTagRedirectors::RefreshTagRedirects()
 {
 	TagRedirects.Empty();
 
+	// Check settings object
 	UGameplayTagsSettings* MutableDefault = GetMutableDefault<UGameplayTagsSettings>();
+	AddRedirects(MutableDefault->GameplayTagRedirects);
+
+	// Add redirects from all TagList sources
+	UGameplayTagsManager& GameplayTagsManager = UGameplayTagsManager::Get();
+	TArray<FGameplayTagRedirect> RedirectsFromTagSources;
+	TArray<const FGameplayTagSource*> TagListSources;
+	GameplayTagsManager.FindTagSourcesWithType(EGameplayTagSourceType::TagList, TagListSources);
+	for (const FGameplayTagSource* Source : TagListSources)
+	{
+		AddRedirectsFromSource(Source);
+	}
 
-	// Check settings object
-	for (const FGameplayTagRedirect& Redirect : MutableDefault->GameplayTagRedirects)
+}
+
+const FGameplayTag* FGameplayTagRedirectors::RedirectTag(const FName& InTagName) const
+{
+	if (const FGameplayTag* NewTag = TagRedirects.Find(InTagName))
+	{
+		return NewTag;
+	}
+
+	return nullptr;
+}
+
+void FGameplayTagRedirectors::AddRedirects(const TArray<FGameplayTagRedirect>& Redirects)
+{
+	for (const FGameplayTagRedirect& Redirect : Redirects)
 	{
 		FName OldTagName = Redirect.OldTagName;
 		FName NewTagName = Redirect.NewTagName;
 
-		if (ensureMsgf(!TagRedirects.Contains(OldTagName), TEXT("Old tag %s is being redirected to more than one tag. Please remove all the redirections except for one."), *OldTagName.ToString()))
+		const FGameplayTag* ExistingRedirect = TagRedirects.Find(OldTagName);
+		if (!ExistingRedirect)
 		{
 			// Attempt to find multiple redirect hops and flatten the redirection so we only need to redirect once
 			// to resolve the update.  Includes a basic infinite recursion guard, in case the redirects loop.
@@ -82,7 +117,7 @@ void FGameplayTagRedirectors::RefreshTagRedirects()
 				bool bFoundRedirect = false;
 
 				// See if it got redirected again
-				for (const FGameplayTagRedirect& SecondRedirect : MutableDefault->GameplayTagRedirects)
+				for (const FGameplayTagRedirect& SecondRedirect : Redirects)
 				{
 					if (SecondRedirect.OldTagName == NewTagName)
 					{
@@ -108,16 +143,9 @@ void FGameplayTagRedirectors::RefreshTagRedirects()
 			// Populate the map
 			TagRedirects.Add(OldTagName, FGameplayTag(NewTagName));
 		}
+		else
+		{
+			ensureMsgf(ExistingRedirect->GetTagName() == NewTagName, TEXT("Old tag %s is being redirected to more than one tag. Please remove all the redirections except for one."), *OldTagName.ToString());
+		}
 	}
-}
-
-const FGameplayTag* FGameplayTagRedirectors::RedirectTag(const FName& InTagName) const
-{
-	if (const FGameplayTag* NewTag = TagRedirects.Find(InTagName))
-	{
-		return NewTag;
-	}
-
-	return nullptr;
-}
-
+}
\ No newline at end of file
diff --git a/Engine/Source/Runtime/GameplayTags/Private/GameplayTagsManager.cpp b/Engine/Source/Runtime/GameplayTags/Private/GameplayTagsManager.cpp
--- a/Engine/Source/Runtime/GameplayTags/Private/GameplayTagsManager.cpp
+++ b/Engine/Source/Runtime/GameplayTags/Private/GameplayTagsManager.cpp
@@ -402,6 +402,8 @@ void UGameplayTagsManager::AddTagsFromAdditionalLooseIniFiles(const TArray<FStri
 			// since LoadConfig can't take an existing FConfigFile* to load from, we put it into GConfig, then remove it
 			GConfig->Remove(IniFilePath);
 
+			FGameplayTagRedirectors::Get().AddRedirectsFromSource(FoundSource);
+
 #if WITH_EDITOR
 			if (GIsEditor || IsRunningCommandlet()) // Sort tags for UI Purposes but don't sort in -game scenario since this would break compat with noneditor cooked builds
 			{
diff --git a/Engine/Source/Runtime/GameplayTags/Public/GameplayTagRedirectors.h b/Engine/Source/Runtime/GameplayTags/Public/GameplayTagRedirectors.h
--- a/Engine/Source/Runtime/GameplayTags/Public/GameplayTagRedirectors.h
+++ b/Engine/Source/Runtime/GameplayTags/Public/GameplayTagRedirectors.h
@@ -45,12 +45,18 @@ public:
 	/** Sees if the tag name should be redirected to a different tag, returns null if there is no active redirect */
 	const FGameplayTag* RedirectTag(const FName& InTagName) const;
 
+	/** Adds redirects from the specified source. Faster than refreshing the entire list */
+	void AddRedirectsFromSource(const struct FGameplayTagSource* Source);
+
 	/** Refreshes the redirect map after a config change */
 	void RefreshTagRedirects();
 
 private:
 	FGameplayTagRedirectors();
 
+	/** Adds a list of tag redirects to the map */
+	void AddRedirects(const TArray<FGameplayTagRedirect>& Redirects);
+
 	/** The map of ini-configured tag redirectors */
 	TMap<FName, FGameplayTag> TagRedirects;
 };
