#UE Adding more info to an ensure

diff --git a/Engine/Source/Runtime/GameplayTags/Private/GameplayTagRedirectors.cpp b/Engine/Source/Runtime/GameplayTags/Private/GameplayTagRedirectors.cpp
--- a/Engine/Source/Runtime/GameplayTags/Private/GameplayTagRedirectors.cpp
+++ b/Engine/Source/Runtime/GameplayTags/Private/GameplayTagRedirectors.cpp
@@ -145,7 +145,7 @@ void FGameplayTagRedirectors::AddRedirects(const TArray<FGameplayTagRedirect>& R
 		}
 		else
 		{
-			ensureMsgf(ExistingRedirect->GetTagName() == NewTagName, TEXT("Old tag %s is being redirected to more than one tag. Please remove all the redirections except for one."), *OldTagName.ToString());
+			ensureMsgf(ExistingRedirect->GetTagName() == NewTagName, TEXT("Old tag %s is being redirected to more than one tag. Please remove all the redirections except for one. NewTagName:%s ExistingRedirect:%s"), *OldTagName.ToString(), *NewTagName.ToString(), *ExistingRedirect->GetTagName().ToString());
 		}
 	}
 }
\ No newline at end of file
