GameplayTagsManager will not ensure against Fast Tag Replication network indices size when not using Fast Tag Replication.  We achieve this by not reconstructing the Net Indices when using non-Fast Tag Replication (e.g. the new Dynamic Tag Replication).  This allows Dynamic Replication to continue past the size of uint16 Gameplay Tags without workflow issues.

diff --git a/Engine/Source/Runtime/GameplayTags/Classes/GameplayTagsManager.h b/Engine/Source/Runtime/GameplayTags/Classes/GameplayTagsManager.h
--- a/Engine/Source/Runtime/GameplayTags/Classes/GameplayTagsManager.h
+++ b/Engine/Source/Runtime/GameplayTags/Classes/GameplayTagsManager.h
@@ -874,7 +874,11 @@ private:
 
 	void VerifyNetworkIndex() const
 	{
-		if (bNetworkIndexInvalidated)
+		if (!bUseFastReplication)
+		{
+			UE_LOG(LogGameplayTags, Warning, TEXT("%hs called when not using FastReplication (not rebuilding the fast replication cache)"), __func__);
+		}
+		else if (bNetworkIndexInvalidated)
 		{
 			const_cast<UGameplayTagsManager*>(this)->ConstructNetIndex();
 		}
