Make use of the Binary Plugin ConfigCache over using disk I/O

diff --git a/Engine/Plugins/Runtime/GameFeatures/Source/GameFeatures/Private/GameFeaturePluginStateMachine.cpp b/Engine/Plugins/Runtime/GameFeatures/Source/GameFeatures/Private/GameFeaturePluginStateMachine.cpp
--- a/Engine/Plugins/Runtime/GameFeatures/Source/GameFeatures/Private/GameFeaturePluginStateMachine.cpp
+++ b/Engine/Plugins/Runtime/GameFeatures/Source/GameFeatures/Private/GameFeaturePluginStateMachine.cpp
@@ -2983,7 +2983,7 @@ struct FGameFeaturePluginState_Registering : public FGameFeaturePluginState
 
 		if (AllowIniLoading())
 		{
-			UGameplayTagsManager::Get().AddTagIniSearchPath(PluginFolder / TEXT("Config") / TEXT("Tags"));
+			UGameplayTagsManager::Get().AddTagIniSearchPath(PluginFolder / TEXT("Config") / TEXT("Tags"), GConfig->GetStagedPluginConfigCache(FName(*StateProperties.PluginName)));
 		}
 
 		LoadGFDState = ELoadGFDState::Pending;
diff --git a/Engine/Source/Runtime/Core/Private/Misc/ConfigCacheIni.cpp b/Engine/Source/Runtime/Core/Private/Misc/ConfigCacheIni.cpp
--- a/Engine/Source/Runtime/Core/Private/Misc/ConfigCacheIni.cpp
+++ b/Engine/Source/Runtime/Core/Private/Misc/ConfigCacheIni.cpp
@@ -6555,6 +6555,15 @@ void FConfigCacheIni::RemoveTagsFromBranches(const TArray<FName>& Tags, FConfigM
 	}
 }
 
+const TSet<FString>* FConfigCacheIni::GetStagedPluginConfigCache(FName PluginName) const
+{
+	return StagedPluginConfigCache.Find(PluginName);
+}
+
+const TSet<FString>* FConfigCacheIni::GetStagedGlobalConfigCache() const
+{
+	return StagedGlobalConfigCache;
+}
 
 #if ALLOW_OTHER_PLATFORM_CONFIG
 
diff --git a/Engine/Source/Runtime/Core/Public/Misc/ConfigCacheIni.h b/Engine/Source/Runtime/Core/Public/Misc/ConfigCacheIni.h
--- a/Engine/Source/Runtime/Core/Public/Misc/ConfigCacheIni.h
+++ b/Engine/Source/Runtime/Core/Public/Misc/ConfigCacheIni.h
@@ -1912,6 +1912,17 @@ public:
 	*/
 	static CORE_API void PreInitializePlatformPlugins();
 
+	/**
+	 * Finds the Staged Plugin ConfigCache if it exists from a PluginName
+	 *   Returns a pointer to the cache if it exists, a nullptr if not
+	 */
+	CORE_API const TSet<FString>* GetStagedPluginConfigCache(FName PluginName) const;
+
+	/**
+	 * Returns the staged global config cache
+	 */
+	CORE_API const TSet<FString>* GetStagedGlobalConfigCache() const;
+
 private:
 #if WITH_EDITOR
 	/** We only auto-initialize other platform configs in the editor to not slow down programs like ShaderCOmpileWorker */
diff --git a/Engine/Source/Runtime/GameplayTags/Classes/GameplayTagsManager.h b/Engine/Source/Runtime/GameplayTags/Classes/GameplayTagsManager.h
--- a/Engine/Source/Runtime/GameplayTags/Classes/GameplayTagsManager.h
+++ b/Engine/Source/Runtime/GameplayTags/Classes/GameplayTagsManager.h
@@ -489,8 +489,8 @@ public:
 	/** Loads the tag tables referenced in the GameplayTagSettings object */
 	GAMEPLAYTAGS_API void LoadGameplayTagTables(bool bAllowAsyncLoad = false);
 
-	/** Loads tag inis contained in the specified path */
-	GAMEPLAYTAGS_API void AddTagIniSearchPath(const FString& RootDir);
+	/** Loads tag inis contained in the specified path, passes an optional PluginConfigCache to speed up disk searches */
+	GAMEPLAYTAGS_API void AddTagIniSearchPath(const FString& RootDir, const TSet<FString>* PluginConfigsCache = nullptr);
 
 	/** Tries to remove the specified search path, will return true if anything was removed */
 	GAMEPLAYTAGS_API bool RemoveTagIniSearchPath(const FString& RootDir);
@@ -970,4 +970,4 @@ private:
 
 	GAMEPLAYTAGS_API const static FName NAME_Categories;
 	GAMEPLAYTAGS_API const static FName NAME_GameplayTagFilter;
-};
\ No newline at end of file
+};
diff --git a/Engine/Source/Runtime/GameplayTags/Private/GameplayTagsManager.cpp b/Engine/Source/Runtime/GameplayTags/Private/GameplayTagsManager.cpp
--- a/Engine/Source/Runtime/GameplayTags/Private/GameplayTagsManager.cpp
+++ b/Engine/Source/Runtime/GameplayTags/Private/GameplayTagsManager.cpp
@@ -246,7 +246,7 @@ void UGameplayTagsManager::LoadGameplayTagTables(bool bAllowAsyncLoad)
 	}
 }
 
-void UGameplayTagsManager::AddTagIniSearchPath(const FString& RootDir)
+void UGameplayTagsManager::AddTagIniSearchPath(const FString& RootDir, const TSet<FString>* PluginConfigsCache)
 {
 	SCOPE_SECONDS_ACCUMULATOR(STAT_GameplayTags_AddTagIniSearchPath);
 
@@ -262,21 +262,37 @@ void UGameplayTagsManager::AddTagIniSearchPath(const FString& RootDir)
 		PathInfo->Reset();
 		
 		// Read all tags from the ini
-		TArray<FString> FilesInDirectory;
-		IFileManager::Get().FindFilesRecursive(FilesInDirectory, *RootDir, TEXT("*.ini"), true, false);
-
-		if (FilesInDirectory.Num() > 0)
+		// Use slower path and check the filesystem if our PluginConfigsCache is null
+		if (PluginConfigsCache == nullptr)
 		{
-			FilesInDirectory.Sort();
+			TArray<FString> FilesInDirectory;
+			IFileManager::Get().FindFilesRecursive(FilesInDirectory, *RootDir, TEXT("*.ini"), true, false);
 
-			for (const FString& IniFilePath : FilesInDirectory)
+			if (FilesInDirectory.Num() > 0)
 			{
-				const FName TagSource = FName(*FPaths::GetCleanFilename(IniFilePath));
-				PathInfo->SourcesInPath.Add(TagSource);
-				PathInfo->TagIniList.Add(FConfigCacheIni::NormalizeConfigIniPath(IniFilePath));
+				FilesInDirectory.Sort();
+
+				for (const FString& IniFilePath : FilesInDirectory)
+				{
+					const FName TagSource = FName(*FPaths::GetCleanFilename(IniFilePath));
+					PathInfo->SourcesInPath.Add(TagSource);
+					PathInfo->TagIniList.Add(FConfigCacheIni::NormalizeConfigIniPath(IniFilePath));
+				}
+			}
+		}
+		else
+		{
+			for (const FString& IniFilePath : *PluginConfigsCache)
+			{
+				// Only grab ini files that part of the root dir we are looking for
+				if (IniFilePath.StartsWith(RootDir))
+				{
+					const FName TagSource = FName(*FPaths::GetCleanFilename(IniFilePath));
+					PathInfo->SourcesInPath.Add(TagSource);
+					PathInfo->TagIniList.Add(FConfigCacheIni::NormalizeConfigIniPath(IniFilePath));
+				}
 			}
 		}
-		
 		PathInfo->bWasSearched = true;
 	}
 
