Dragging from a Key pin to a FText pin will now automatically add a Get Key Display Name node.
Dragging from a GameplayTag or GameplayTagContainer pin to a String pin will now automatically add a GetDebugString node.

diff --git a/Engine/Source/Runtime/Engine/Classes/Kismet/KismetInputLibrary.h b/Engine/Source/Runtime/Engine/Classes/Kismet/KismetInputLibrary.h
--- a/Engine/Source/Runtime/Engine/Classes/Kismet/KismetInputLibrary.h
+++ b/Engine/Source/Runtime/Engine/Classes/Kismet/KismetInputLibrary.h
@@ -154,7 +154,7 @@ class UKismetInputLibrary : public UBlueprintFunctionLibrary
 	/**
 	 * Returns the display name of the key.
 	 */
-	UFUNCTION(BlueprintPure, meta = (DisplayName = "Get Key Display Name"), Category="Input|Key")
+	UFUNCTION(BlueprintPure, meta = (DisplayName = "Get Key Display Name", BlueprintAutocast), Category="Input|Key")
 	static ENGINE_API FText Key_GetDisplayName(const FKey& Key, bool bLongDisplayName = true);
 
 	/**
diff --git a/Engine/Source/Runtime/GameplayTags/Classes/BlueprintGameplayTagLibrary.h b/Engine/Source/Runtime/GameplayTags/Classes/BlueprintGameplayTagLibrary.h
--- a/Engine/Source/Runtime/GameplayTags/Classes/BlueprintGameplayTagLibrary.h
+++ b/Engine/Source/Runtime/GameplayTags/Classes/BlueprintGameplayTagLibrary.h
@@ -282,7 +282,7 @@ class UBlueprintGameplayTagLibrary : public UBlueprintFunctionLibrary
 	 *
 	 * @param TagContainer	The tag container to get the debug string from.
 	 */
-	UFUNCTION(BlueprintPure, Category = "GameplayTags", meta = (BlueprintThreadSafe))
+	UFUNCTION(BlueprintPure, Category = "GameplayTags", meta = (BlueprintThreadSafe, BlueprintAutocast))
 	static GAMEPLAYTAGS_API FString GetDebugStringFromGameplayTagContainer(const FGameplayTagContainer& TagContainer);
 
 	/**
@@ -290,7 +290,7 @@ class UBlueprintGameplayTagLibrary : public UBlueprintFunctionLibrary
 	 *
 	 * @param GameplayTag	The tag to get the debug string from.
 	 */
-	UFUNCTION(BlueprintPure, Category = "GameplayTags", meta = (BlueprintThreadSafe))
+	UFUNCTION(BlueprintPure, Category = "GameplayTags", meta = (BlueprintThreadSafe, BlueprintAutocast))
 	static GAMEPLAYTAGS_API FString GetDebugStringFromGameplayTag(FGameplayTag GameplayTag);
 
 };
