Fixed typo in ensure message of GameplayTagContainer

diff --git a/Engine/Source/Runtime/GameplayTags/Private/GameplayTagContainer.cpp b/Engine/Source/Runtime/GameplayTags/Private/GameplayTagContainer.cpp
--- a/Engine/Source/Runtime/GameplayTags/Private/GameplayTagContainer.cpp
+++ b/Engine/Source/Runtime/GameplayTags/Private/GameplayTagContainer.cpp
@@ -513,7 +513,7 @@ bool FQueryEvaluator::EvalExpr(FGameplayTagContainer const& Tags, bool bSkip)
 		return EvalNoExprMatch(Tags, bSkip);
 	}
 
-	ensureAlwaysMsgf(false, TEXT("Encounted invalid query expression type: %s"), *UEnum::GetValueAsString(ExprType));
+	ensureAlwaysMsgf(false, TEXT("Encountered invalid query expression type: %s"), *UEnum::GetValueAsString(ExprType));
 	return false;
 }
 
