UGameplayTagsManager::IsValidGameplayTagString - Add overloads for using FStringView and FStringBuilderBase. Address cases where the value of OutFixedString would not be valid if TagString was an invalid gameplay tag string of a particular format.

#tests All GameplayTagManager tests pass

diff --git a/Engine/Source/Runtime/GameplayTags/Classes/GameplayTagsManager.h b/Engine/Source/Runtime/GameplayTags/Classes/GameplayTagsManager.h
--- a/Engine/Source/Runtime/GameplayTags/Classes/GameplayTagsManager.h
+++ b/Engine/Source/Runtime/GameplayTags/Classes/GameplayTagsManager.h
@@ -358,7 +358,9 @@ class UGameplayTagsManager : public UObject
 	 * @param OutFixedString If non-null and string invalid, will attempt to fix. Will be empty if no fix is possible
 	 * @return True if this can be added to the tag dictionary, false if there's a syntax error
 	 */
+	GAMEPLAYTAGS_API bool IsValidGameplayTagString(const TCHAR* TagString, FText* OutError = nullptr, FString* OutFixedString = nullptr);
 	GAMEPLAYTAGS_API bool IsValidGameplayTagString(const FString& TagString, FText* OutError = nullptr, FString* OutFixedString = nullptr);
+	GAMEPLAYTAGS_API bool IsValidGameplayTagString(const FStringView& TagString, FText* OutError = nullptr, FStringBuilderBase* OutFixedString = nullptr);
 
 	/**
 	 *	Searches for a gameplay tag given a partial string. This is slow and intended mainly for console commands/utilities to make
diff --git a/Engine/Source/Runtime/GameplayTags/Private/GameplayTagsManager.cpp b/Engine/Source/Runtime/GameplayTags/Private/GameplayTagsManager.cpp
--- a/Engine/Source/Runtime/GameplayTags/Private/GameplayTagsManager.cpp
+++ b/Engine/Source/Runtime/GameplayTags/Private/GameplayTagsManager.cpp
@@ -18,6 +18,7 @@
 #include "GameplayTagsModule.h"
 #include "Framework/Notifications/NotificationManager.h"
 #include "Widgets/Notifications/SNotificationList.h"
+#include "Misc/AsciiSet.h"
 #include "Misc/CoreDelegates.h"
 #include "HAL/IConsoleManager.h"
 #include "NativeGameplayTags.h"
@@ -2038,75 +2039,127 @@ FGameplayTag UGameplayTagsManager::RequestGameplayTag(FName TagName, bool ErrorI
 	return FGameplayTag();
 }
 
-bool UGameplayTagsManager::IsValidGameplayTagString(const FString& TagString, FText* OutError, FString* OutFixedString)
+namespace UE::GameplayTags::Private
+{
+
+template <typename FixedStringType>
+bool IsValidGameplayTagString(const FStringView& TagString, FText* OutError, FixedStringType* OutFixedString, const FString& InvalidTagCharacters)
 {
 	bool bIsValid = true;
-	FString FixedString = TagString;
-	FText ErrorText;
+	FStringView FixedString = TagString;
+	TArray<FText> Errors;
 
 	if (FixedString.IsEmpty())
 	{
-		ErrorText = LOCTEXT("EmptyStringError", "Tag is empty");
-		bIsValid = false;
-	}
-
-	while (FixedString.StartsWith(TEXT("."), ESearchCase::CaseSensitive))
-	{
-		ErrorText = LOCTEXT("StartWithPeriod", "Tag starts with .");
-		FixedString.RemoveAt(0);
+		Errors.Add(LOCTEXT("EmptyStringError", "Tag may not be empty"));
 		bIsValid = false;
 	}
 
-	while (FixedString.EndsWith(TEXT("."), ESearchCase::CaseSensitive))
-	{
-		ErrorText = LOCTEXT("EndWithPeriod", "Tag ends with .");
-		FixedString.RemoveAt(FixedString.Len() - 1);
-		bIsValid = false;
-	}
+	constexpr FAsciiSet Period(".");
+	constexpr FAsciiSet Space(" ");
+	constexpr FAsciiSet TrimSet = Period | Space;
 
-	while (FixedString.StartsWith(TEXT(" "), ESearchCase::CaseSensitive))
+	if (const FStringView Trimmed = FAsciiSet::TrimPrefixWith(FixedString, TrimSet); Trimmed.Len() != FixedString.Len())
 	{
-		ErrorText = LOCTEXT("StartWithSpace", "Tag starts with space");
-		FixedString.RemoveAt(0);
+		if (OutError)
+		{
+			const FStringView Removed = FixedString.LeftChop(Trimmed.Len());
+			if (FAsciiSet::HasAny(Removed, Period))
+			{
+				Errors.Add(LOCTEXT("StartWithPeriod", "Tag may not begin with a period ('.')"));
+			}
+			if (FAsciiSet::HasAny(Removed, Space))
+			{
+				Errors.Add(LOCTEXT("StartWithSpace", "Tag may not begin with a space"));
+			}
+		}
 		bIsValid = false;
+		FixedString = Trimmed;
 	}
 
-	while (FixedString.EndsWith(TEXT(" "), ESearchCase::CaseSensitive))
+	if (const FStringView Trimmed = FAsciiSet::TrimSuffixWith(FixedString, TrimSet); Trimmed.Len() != FixedString.Len())
 	{
-		ErrorText = LOCTEXT("EndWithSpace", "Tag ends with space");
-		FixedString.RemoveAt(FixedString.Len() - 1);
+		if (OutError)
+		{
+			const FStringView Removed = FixedString.RightChop(Trimmed.Len());
+			if (FAsciiSet::HasAny(Removed, Period))
+			{
+				Errors.Add(LOCTEXT("EndWithPeriod", "Tag may not end with a period ('.')"));
+			}
+			if (FAsciiSet::HasAny(Removed, Space))
+			{
+				Errors.Add(LOCTEXT("EndWithSpace", "Tag may not end with a space"));
+			}
+		}
 		bIsValid = false;
+		FixedString = Trimmed;
 	}
 
-	FText TagContext = LOCTEXT("GameplayTagContext", "Tag");
-	if (!FName::IsValidXName(TagString, InvalidTagCharacters, &ErrorText, &TagContext))
+	FText ErrorText;
+	const FText TagContext = LOCTEXT("GameplayTagContext", "Tag");
+	if (!FName::IsValidXName(FixedString, InvalidTagCharacters, OutError ? &ErrorText : nullptr, &TagContext))
 	{
-		for (TCHAR& TestChar : FixedString)
+		if (OutError)
+		{
+			Errors.Add(MoveTemp(ErrorText));
+		}
+		if (OutFixedString)
 		{
-			for (TCHAR BadChar : InvalidTagCharacters)
+			OutFixedString->Reset();
+			OutFixedString->Reserve(FixedString.Len());
+			for (const TCHAR& Char : FixedString)
 			{
-				if (TestChar == BadChar)
+				int32 CharIndex;
+				if (InvalidTagCharacters.FindChar(Char, CharIndex))
+				{
+					*OutFixedString += TEXT('_');
+				}
+				else
 				{
-					TestChar = TEXT('_');
+					*OutFixedString += Char;
 				}
 			}
 		}
-
 		bIsValid = false;
 	}
-
-	if (OutError)
+	else if (OutFixedString)
 	{
-		*OutError = ErrorText;
+		OutFixedString->Reset();
+		*OutFixedString += FixedString;
 	}
-	if (OutFixedString)
+
+	if (OutError)
 	{
-		*OutFixedString = FixedString;
+		if (!Errors.IsEmpty())
+		{
+			*OutError = FText::Join(LOCTEXT("ErrorDelimiter", ", "), Errors);
+		}
+		else
+		{
+			*OutError = FText();
+		}
 	}
 
 	return bIsValid;
 }
 
+} // namespace UE::GameplayTags::Private
+
+bool UGameplayTagsManager::IsValidGameplayTagString(const TCHAR* TagString, FText* OutError /*= nullptr*/, FString* OutFixedString /*= nullptr*/)
+{
+	return UE::GameplayTags::Private::IsValidGameplayTagString(FStringView(TagString), OutError, OutFixedString, InvalidTagCharacters);
+}
+
+bool UGameplayTagsManager::IsValidGameplayTagString(const FString& TagString, FText* OutError /*= nullptr*/, FString* OutFixedString /*= nullptr*/)
+{
+	return UE::GameplayTags::Private::IsValidGameplayTagString(FStringView(TagString), OutError, OutFixedString, InvalidTagCharacters);
+}
+
+bool UGameplayTagsManager::IsValidGameplayTagString(const FStringView& TagString, FText* OutError /*= nullptr*/, FStringBuilderBase* OutFixedString /*= nullptr*/)
+{
+	return UE::GameplayTags::Private::IsValidGameplayTagString(FStringView(TagString), OutError, OutFixedString, InvalidTagCharacters);
+}
+
 FGameplayTag UGameplayTagsManager::FindGameplayTagFromPartialString_Slow(FString PartialString) const
 {
 	// This critical section is to handle an issue where tag requests come from another thread when async loading from a background thread in FGameplayTagContainer::Serialize.
diff --git a/Engine/Source/Runtime/GameplayTags/Private/Tests/GameplayTagTests.cpp b/Engine/Source/Runtime/GameplayTags/Private/Tests/GameplayTagTests.cpp
--- a/Engine/Source/Runtime/GameplayTags/Private/Tests/GameplayTagTests.cpp
+++ b/Engine/Source/Runtime/GameplayTags/Private/Tests/GameplayTagTests.cpp
@@ -1,6 +1,5 @@
 // Copyright Epic Games, Inc. All Rights Reserved.
-
-#include "CoreMinimal.h"
+#if WITH_TESTS
 #include "UObject/Package.h"
 #include "Misc/AutomationTest.h"
 #include "Engine/DataTable.h"
@@ -8,6 +7,10 @@
 #include "GameplayTagsManager.h"
 #include "GameplayTagsModule.h"
 #include "Stats/StatsMisc.h"
+#include "Tests/TestHarnessAdapter.h"
+#if WITH_LOW_LEVEL_TESTS
+#include "TestCommon/Expectations.h"
+#endif
 
 #if WITH_AUTOMATION_WORKER
 
@@ -416,3 +419,98 @@ bool FGameplayTagTest::RunTest(const FString& Parameters)
 }
 
 #endif //WITH_AUTOMATION_WORKER
+
+TEST_CASE_NAMED(
+	FSystem__GameplayTags__Manager__IsValidGameplayTagString,
+	"System::GameplayTags::Manager::IsValidGameplayTagString",
+	"[ApplicationContextMask][EngineFilter]")
+{
+	UGameplayTagsManager& Manager = UGameplayTagsManager::Get();
+
+	SECTION("Valid tag: no modifications required")
+	{
+		FText Error;
+		bool bValid = Manager.IsValidGameplayTagString(TEXT("Valid.Tag"), &Error);
+		CHECK(bValid);
+		CHECK(Error.IsEmpty());
+	}
+
+	SECTION("Tag with a leading period")
+	{
+		FText Error;
+		FString FixedTag;
+		bool bValid = Manager.IsValidGameplayTagString(TEXT(".Tag"), &Error, &FixedTag);
+		CHECK(!bValid);
+		CHECK_FALSE(Error.IsEmpty());
+		INFO(Error.ToString());
+		CHECK(Manager.IsValidGameplayTagString(FixedTag));
+	}
+
+	SECTION("Tag with a trailing period")
+	{
+		FText Error;
+		FString FixedTag;
+		bool bValid = Manager.IsValidGameplayTagString(TEXT("Tag."), &Error, &FixedTag);
+		CHECK(!bValid);
+		CHECK_FALSE(Error.IsEmpty());
+		INFO(Error.ToString());
+		CHECK(Manager.IsValidGameplayTagString(FixedTag));
+	}
+
+	SECTION("Tag with a leading space")
+	{
+		FText Error;
+		FString FixedTag;
+		bool bValid = Manager.IsValidGameplayTagString(TEXT(" Tag"), &Error, &FixedTag);
+		CHECK(!bValid);
+		CHECK_FALSE(Error.IsEmpty());
+		INFO(Error.ToString());
+		CHECK(Manager.IsValidGameplayTagString(FixedTag));
+	}
+
+	SECTION("Tag with a trailing space")
+	{
+		FText Error;
+		FString FixedTag;
+		bool bValid = Manager.IsValidGameplayTagString(TEXT("Tag "), &Error, &FixedTag);
+		CHECK(!bValid);
+		CHECK_FALSE(Error.IsEmpty());
+		INFO(Error.ToString());
+		CHECK(Manager.IsValidGameplayTagString(FixedTag));
+	}
+
+	SECTION("Tag with both leading and trailing unwanted characters")
+	{
+		FText Error;
+		FString FixedTag;
+		bool bValid = Manager.IsValidGameplayTagString(TEXT(" . Tag . "), &Error, &FixedTag);
+		CHECK(!bValid);
+		CHECK_FALSE(Error.IsEmpty());
+		INFO(Error.ToString());
+		CHECK(Manager.IsValidGameplayTagString(FixedTag));
+	}
+
+	SECTION("Tag with an invalid character inside")
+	{
+		FText Error;
+		FString FixedTag;
+		bool bValid = Manager.IsValidGameplayTagString(TEXT("Tag\t,Name"), &Error, &FixedTag);
+		CHECK(!bValid);
+		CHECK_FALSE(Error.IsEmpty());
+		INFO(Error.ToString());
+		CHECK(Manager.IsValidGameplayTagString(FixedTag));
+	}
+
+	SECTION("Tag with both leading and trailing unwanted characters and invalid characters inside")
+	{
+		FText Error;
+		FString FixedTag;
+		bool bValid = Manager.IsValidGameplayTagString(TEXT(" . Tag\t,Name . "), &Error, &FixedTag);
+		CHECK(!bValid);
+		CHECK_FALSE(Error.IsEmpty());
+		INFO(Error.ToString());
+		CHECK(Manager.IsValidGameplayTagString(FixedTag));
+	}
+}
+
+#endif // WITH_TESTS
\ No newline at end of file
