Changing many GetDefaultMutable to GetDefault in order to move towards banning mutations on CDO's at runtime.
Removed some old (long deprecated) GameplayTagsManager backwards compatability code.

diff --git a/Engine/Plugins/EnhancedInput/Source/EnhancedInput/Private/EnhancedInputSubsystems.cpp b/Engine/Plugins/EnhancedInput/Source/EnhancedInput/Private/EnhancedInputSubsystems.cpp
--- a/Engine/Plugins/EnhancedInput/Source/EnhancedInput/Private/EnhancedInputSubsystems.cpp
+++ b/Engine/Plugins/EnhancedInput/Source/EnhancedInput/Private/EnhancedInputSubsystems.cpp
@@ -319,7 +319,7 @@ void UEnhancedInputWorldSubsystem::AddDefaultMappingContexts()
 {
 	if (GetDefault<UEnhancedInputDeveloperSettings>()->bEnableDefaultMappingContexts)
 	{
-		for (const FDefaultContextSetting& ContextSetting : GetMutableDefault<UEnhancedInputDeveloperSettings>()->DefaultWorldSubsystemMappingContexts)
+		for (const FDefaultContextSetting& ContextSetting : GetDefault<UEnhancedInputDeveloperSettings>()->DefaultWorldSubsystemMappingContexts)
 		{
 			if (const UInputMappingContext* IMC = ContextSetting.InputMappingContext.LoadSynchronous())
 			{
@@ -348,7 +348,7 @@ void UEnhancedInputWorldSubsystem::RemoveDefaultMappingContexts()
 {
 	if (GetDefault<UEnhancedInputDeveloperSettings>()->bEnableDefaultMappingContexts)
 	{
-		for (const FDefaultContextSetting& ContextSetting : GetMutableDefault<UEnhancedInputDeveloperSettings>()->DefaultWorldSubsystemMappingContexts)
+		for (const FDefaultContextSetting& ContextSetting : GetDefault<UEnhancedInputDeveloperSettings>()->DefaultWorldSubsystemMappingContexts)
 		{
 			if (const UInputMappingContext* IMC = ContextSetting.InputMappingContext.LoadSynchronous())
 			{
diff --git a/Engine/Plugins/Messaging/UdpMessaging/Source/UdpMessaging/Private/UdpMessagingModule.cpp b/Engine/Plugins/Messaging/UdpMessaging/Source/UdpMessaging/Private/UdpMessagingModule.cpp
--- a/Engine/Plugins/Messaging/UdpMessaging/Source/UdpMessaging/Private/UdpMessagingModule.cpp
+++ b/Engine/Plugins/Messaging/UdpMessaging/Source/UdpMessaging/Private/UdpMessagingModule.cpp
@@ -46,7 +46,7 @@ protected:
 
 		if (FParse::Command(&Cmd, TEXT("STATUS")))
 		{
-			UUdpMessagingSettings* Settings = GetMutableDefault<UUdpMessagingSettings>();
+			const UUdpMessagingSettings* Settings = GetDefault<UUdpMessagingSettings>();
 
 			// general information
 			Ar.Logf(TEXT("Protocol Version: %i"), UDP_MESSAGING_TRANSPORT_PROTOCOL_VERSION);
diff --git a/Engine/Source/Developer/Localization/Private/LocalizationSettings.cpp b/Engine/Source/Developer/Localization/Private/LocalizationSettings.cpp
--- a/Engine/Source/Developer/Localization/Private/LocalizationSettings.cpp
+++ b/Engine/Source/Developer/Localization/Private/LocalizationSettings.cpp
@@ -78,14 +78,14 @@ void ULocalizationSettings::PostEditChangeProperty(FPropertyChangedEvent& Proper
 
 ULocalizationTargetSet* ULocalizationSettings::GetEngineTargetSet()
 {
-	ULocalizationSettings* LocalizationSettings = GetMutableDefault<ULocalizationSettings>();
+	const ULocalizationSettings* LocalizationSettings = GetDefault<ULocalizationSettings>();
 	check(LocalizationSettings);
 	return LocalizationSettings->EngineTargetSet;
 }
 
 ULocalizationTargetSet* ULocalizationSettings::GetGameTargetSet()
 {
-	ULocalizationSettings* LocalizationSettings = GetMutableDefault<ULocalizationSettings>();
+	const ULocalizationSettings* LocalizationSettings = GetDefault<ULocalizationSettings>();
 	check(LocalizationSettings);
 	return LocalizationSettings->GameTargetSet;
 }
diff --git a/Engine/Source/Runtime/Engine/Classes/Engine/UserInterfaceSettings.h b/Engine/Source/Runtime/Engine/Classes/Engine/UserInterfaceSettings.h
--- a/Engine/Source/Runtime/Engine/Classes/Engine/UserInterfaceSettings.h
+++ b/Engine/Source/Runtime/Engine/Classes/Engine/UserInterfaceSettings.h
@@ -218,7 +218,7 @@ public:
 	ENGINE_API virtual void PostInitProperties() override;
 
 	/** Loads assets, if bForceLoadEverything is true it will load despite environment */
-	ENGINE_API void ForceLoadResources(bool bForceLoadEverything = false);
+	ENGINE_API void ForceLoadResources(bool bForceLoadEverything = false) const;
 
 	/** Gets the current scale of the UI based on the size of a viewport */
 	ENGINE_API float GetDPIScaleBasedOnSize(FIntPoint Size) const;
@@ -274,8 +274,9 @@ public:
 	bool bEnableDistanceFieldFontRasterization;
 
 private:
+	// This is just holding on to references to prevent GC in the Editor (tied to the lifetime of this object)
 	UPROPERTY(Transient)
-	TArray<TObjectPtr<UObject>> CursorClasses;
+	mutable TArray<TObjectPtr<UObject>> CursorClasses;
 
 	UPROPERTY(Transient)
 	mutable TObjectPtr<UClass> CustomScalingRuleClassInstance;
diff --git a/Engine/Source/Runtime/Engine/Private/GameViewportClient.cpp b/Engine/Source/Runtime/Engine/Private/GameViewportClient.cpp
--- a/Engine/Source/Runtime/Engine/Private/GameViewportClient.cpp
+++ b/Engine/Source/Runtime/Engine/Private/GameViewportClient.cpp
@@ -548,7 +548,7 @@ void UGameViewportClient::Init(struct FWorldContext& WorldContext, UGameInstance
 	}
 
 	// Create the cursor Widgets
-	UUserInterfaceSettings* UISettings = GetMutableDefault<UUserInterfaceSettings>(UUserInterfaceSettings::StaticClass());
+	const UUserInterfaceSettings* UISettings = GetDefault<UUserInterfaceSettings>(UUserInterfaceSettings::StaticClass());
 
 	if (GEngine)
 	{
@@ -598,7 +598,7 @@ void UGameViewportClient::Init(struct FWorldContext& WorldContext, UGameInstance
 
 void UGameViewportClient::RebuildCursors()
 {
-	UUserInterfaceSettings* UISettings = GetMutableDefault<UUserInterfaceSettings>(UUserInterfaceSettings::StaticClass());
+	const UUserInterfaceSettings* UISettings = GetDefault<UUserInterfaceSettings>(UUserInterfaceSettings::StaticClass());
 	// Set all the software cursors.
 	for (auto& Entry : UISettings->SoftwareCursors)
 	{
@@ -2270,7 +2270,7 @@ void UGameViewportClient::Precache()
 
 TOptional<bool> UGameViewportClient::QueryShowFocus(const EFocusCause InFocusCause) const
 {
-	UUserInterfaceSettings* UISettings = GetMutableDefault<UUserInterfaceSettings>(UUserInterfaceSettings::StaticClass());
+	const UUserInterfaceSettings* UISettings = GetDefault<UUserInterfaceSettings>(UUserInterfaceSettings::StaticClass());
 
 	if ( UISettings->RenderFocusRule == ERenderFocusRule::Never ||
 		(UISettings->RenderFocusRule == ERenderFocusRule::NonPointer && InFocusCause == EFocusCause::Mouse) ||
diff --git a/Engine/Source/Runtime/Engine/Private/PlayerCameraManager.cpp b/Engine/Source/Runtime/Engine/Private/PlayerCameraManager.cpp
--- a/Engine/Source/Runtime/Engine/Private/PlayerCameraManager.cpp
+++ b/Engine/Source/Runtime/Engine/Private/PlayerCameraManager.cpp
@@ -1129,7 +1129,7 @@ TScriptInterface<class ICameraLensEffectInterface> APlayerCameraManager::FindGen
 		const UObject* LensEffectObject = LensEffectInterface.GetObject();
 
 		// we have to use GetMutableDefault here because TScriptInterface cannot handle a const UObject* and requires a non-const qualified pointer.
-		const TScriptInterface<class ICameraLensEffectInterface> OtherEffectDefaultInterface = GetMutableDefault<AActor>(LensEffectEmitterClass);
+		const TScriptInterface<const class ICameraLensEffectInterface> OtherEffectDefaultInterface = GetDefault<AActor>(LensEffectEmitterClass);
 
 		// if the lens effect in our list is valid, and either it treats the requested effect as the same
 		// or if the requested effect would treat our existing lens effect as the same...
@@ -1154,7 +1154,7 @@ TScriptInterface<class ICameraLensEffectInterface> APlayerCameraManager::AddGene
 	{
 		TScriptInterface<class ICameraLensEffectInterface> SpawnedLensEffectInterface = NULL;
 
-		const TScriptInterface<class ICameraLensEffectInterface> DesiredLensEffect_DefaultInterface = GetMutableDefault<AActor>(LensEffectEmitterClass);
+		const TScriptInterface<const class ICameraLensEffectInterface> DesiredLensEffect_DefaultInterface = GetDefault<AActor>(LensEffectEmitterClass);
 		const AActor* DesiredLensEffect_DefaultObject = Cast<AActor>(DesiredLensEffect_DefaultInterface.GetObject());
 
 		if (DesiredLensEffect_DefaultInterface && !DesiredLensEffect_DefaultInterface->ShouldAllowMultipleInstances())
diff --git a/Engine/Source/Runtime/Engine/Private/SoundWave.cpp b/Engine/Source/Runtime/Engine/Private/SoundWave.cpp
--- a/Engine/Source/Runtime/Engine/Private/SoundWave.cpp
+++ b/Engine/Source/Runtime/Engine/Private/SoundWave.cpp
@@ -1402,7 +1402,7 @@ ESoundAssetCompressionType USoundWave::GetSoundAssetCompressionType() const
 
 		if (!bDefaultCompressionTypeCached)
 		{
-			UAudioSettings* Settings = GetMutableDefault<UAudioSettings>();
+			const UAudioSettings* Settings = GetDefault<UAudioSettings>();
 			DefaultAudioCompressionType = Audio::ToSoundAssetCompressionType(Settings->DefaultAudioCompressionType);
 			bDefaultCompressionTypeCached = true;
 		}
@@ -3509,7 +3509,7 @@ void USoundWave::Parse(FAudioDevice* AudioDevice, const UPTRINT NodeWaveInstance
 
 			if (SoundClassProperties->bApplyEffects)
 			{
-				UAudioSettings* Settings = GetMutableDefault<UAudioSettings>();
+				const UAudioSettings* Settings = GetDefault<UAudioSettings>();
 				WaveInstance->SoundSubmix = Cast<USoundSubmix>(FSoftObjectPtr(Settings->EQSubmix).Get());
 			}
 			else if (SoundClassProperties->DefaultSubmix)
diff --git a/Engine/Source/Runtime/Engine/Private/UnrealEngine.cpp b/Engine/Source/Runtime/Engine/Private/UnrealEngine.cpp
--- a/Engine/Source/Runtime/Engine/Private/UnrealEngine.cpp
+++ b/Engine/Source/Runtime/Engine/Private/UnrealEngine.cpp
@@ -3471,7 +3471,7 @@ void UEngine::InitializeObjectReferences()
 		}
 	}
 
-	UUserInterfaceSettings* UISettings = GetMutableDefault<UUserInterfaceSettings>(UUserInterfaceSettings::StaticClass());
+	const UUserInterfaceSettings* UISettings = GetDefault<UUserInterfaceSettings>(UUserInterfaceSettings::StaticClass());
 	UISettings->ForceLoadResources();
 }
 
diff --git a/Engine/Source/Runtime/Engine/Private/UserInterfaceSettings.cpp b/Engine/Source/Runtime/Engine/Private/UserInterfaceSettings.cpp
--- a/Engine/Source/Runtime/Engine/Private/UserInterfaceSettings.cpp
+++ b/Engine/Source/Runtime/Engine/Private/UserInterfaceSettings.cpp
@@ -230,7 +230,7 @@ void UUserInterfaceSettings::PostEditChangeProperty(FPropertyChangedEvent& Prope
 
 #endif
 
-void UUserInterfaceSettings::ForceLoadResources(bool bForceLoadEverything)
+void UUserInterfaceSettings::ForceLoadResources(bool bForceLoadEverything) const
 {
 	bool bShouldLoadCursors = true;
 
diff --git a/Engine/Source/Runtime/GameplayTags/Private/GameplayTagRedirectors.cpp b/Engine/Source/Runtime/GameplayTags/Private/GameplayTagRedirectors.cpp
--- a/Engine/Source/Runtime/GameplayTags/Private/GameplayTagRedirectors.cpp
+++ b/Engine/Source/Runtime/GameplayTags/Private/GameplayTagRedirectors.cpp
@@ -74,8 +74,8 @@ void FGameplayTagRedirectors::RefreshTagRedirects()
 	TagRedirects.Empty();
 
 	// Check settings object
-	UGameplayTagsSettings* MutableDefault = GetMutableDefault<UGameplayTagsSettings>();
-	AddRedirects(MutableDefault->GameplayTagRedirects);
+	const UGameplayTagsSettings* Default = GetDefault<UGameplayTagsSettings>();
+	AddRedirects(Default->GameplayTagRedirects);
 
 	// Add redirects from all TagList sources
 	UGameplayTagsManager& GameplayTagsManager = UGameplayTagsManager::Get();
diff --git a/Engine/Source/Runtime/GameplayTags/Private/GameplayTagsManager.cpp b/Engine/Source/Runtime/GameplayTags/Private/GameplayTagsManager.cpp
--- a/Engine/Source/Runtime/GameplayTags/Private/GameplayTagsManager.cpp
+++ b/Engine/Source/Runtime/GameplayTags/Private/GameplayTagsManager.cpp
@@ -211,14 +211,14 @@ void UGameplayTagsManager::LoadGameplayTagTables(bool bAllowAsyncLoad)
 {
 	SCOPE_CYCLE_COUNTER(STAT_GameplayTags_LoadGameplayTags);
 
-	UGameplayTagsSettings* MutableDefault = GetMutableDefault<UGameplayTagsSettings>();
+	const UGameplayTagsSettings* Default = GetDefault<UGameplayTagsSettings>();
 	GameplayTagTables.Empty();
 
 #if !WITH_EDITOR
 	// If we're a cooked build and in a safe spot, start an async load so we can pipeline it
-	if (bAllowAsyncLoad && !IsLoading() && MutableDefault->GameplayTagTableList.Num() > 0)
+	if (bAllowAsyncLoad && !IsLoading() && Default->GameplayTagTableList.Num() > 0)
 	{
-		for (FSoftObjectPath DataTablePath : MutableDefault->GameplayTagTableList)
+		for (FSoftObjectPath DataTablePath : Default->GameplayTagTableList)
 		{
 			LoadPackageAsync(DataTablePath.GetLongPackageName());
 		}
@@ -228,7 +228,7 @@ void UGameplayTagsManager::LoadGameplayTagTables(bool bAllowAsyncLoad)
 #endif // !WITH_EDITOR
 
 	SCOPE_LOG_GAMEPLAYTAGS(TEXT("UGameplayTagsManager::LoadGameplayTagTables"));
-	for (FSoftObjectPath DataTablePath : MutableDefault->GameplayTagTableList)
+	for (FSoftObjectPath DataTablePath : Default->GameplayTagTableList)
 	{
 		UDataTable* TagTable = LoadObject<UDataTable>(nullptr, *DataTablePath.ToString(), nullptr, LOAD_None, nullptr);
 
@@ -444,10 +444,8 @@ void UGameplayTagsManager::ConstructGameplayTagTree()
 	{
 		GameplayRootTag = MakeShareable(new FGameplayTagNode());
 
-		UGameplayTagsSettings* MutableDefault = GetMutableDefault<UGameplayTagsSettings>();
-
 		// Copy invalid characters, then add internal ones
-		InvalidTagCharacters = MutableDefault->InvalidTagCharacters;
+		InvalidTagCharacters = GetDefault<UGameplayTagsSettings>()->InvalidTagCharacters;
 		InvalidTagCharacters.Append(TEXT("\r\n\t"));
 
 		// Add prefixes first
@@ -505,13 +503,15 @@ void UGameplayTagsManager::ConstructGameplayTagTree()
 			SCOPE_LOG_GAMEPLAYTAGS(TEXT("UGameplayTagsManager::ConstructGameplayTagTree: ImportINI tags"));
 
 #if WITH_EDITOR
-			MutableDefault->SortTags();
+			GetMutableDefault<UGameplayTagsSettings>()->SortTags();
 #endif
 
+			const UGameplayTagsSettings* Default = GetDefault<UGameplayTagsSettings>();
+
 			FName TagSource = FGameplayTagSource::GetDefaultName();
 			FGameplayTagSource* DefaultSource = FindOrAddTagSource(TagSource, EGameplayTagSourceType::DefaultTagList);
 
-			for (const FGameplayTagTableRow& TableRow : MutableDefault->GameplayTagList)
+			for (const FGameplayTagTableRow& TableRow : Default->GameplayTagList)
 			{
 				AddTagTableRow(TableRow, TagSource);
 			}
@@ -544,9 +544,12 @@ void UGameplayTagsManager::ConstructGameplayTagTree()
 #endif
 		{
 			SCOPE_LOG_GAMEPLAYTAGS(TEXT("UGameplayTagsManager::ConstructGameplayTagTree: Request common tags"));
+
+			const UGameplayTagsSettings* Default = GetDefault<UGameplayTagsSettings>();
+
 			// Grab the commonly replicated tags
 			CommonlyReplicatedTags.Empty();
-			for (FName TagName : MutableDefault->CommonlyReplicatedTags)
+			for (FName TagName : Default->CommonlyReplicatedTags)
 			{
 				if (TagName.IsNone())
 				{
@@ -565,21 +568,21 @@ void UGameplayTagsManager::ConstructGameplayTagTree()
 				}
 			}
 
-			bUseFastReplication = MutableDefault->FastReplication;
-			bUseDynamicReplication = MutableDefault->bDynamicReplication;
-			bShouldWarnOnInvalidTags = MutableDefault->WarnOnInvalidTags;
-			NumBitsForContainerSize = MutableDefault->NumBitsForContainerSize;
-			NetIndexFirstBitSegment = MutableDefault->NetIndexFirstBitSegment;
+			bUseFastReplication = Default->FastReplication;
+			bUseDynamicReplication = Default->bDynamicReplication;
+			bShouldWarnOnInvalidTags = Default->WarnOnInvalidTags;
+			NumBitsForContainerSize = Default->NumBitsForContainerSize;
+			NetIndexFirstBitSegment = Default->NetIndexFirstBitSegment;
 
 #if WITH_EDITOR
 			if (GIsEditor)
 			{
-				bShouldAllowUnloadingTags = MutableDefault->AllowEditorTagUnloading;
+				bShouldAllowUnloadingTags = Default->AllowEditorTagUnloading;
 			}
 			else
 #endif
 			{
-				bShouldAllowUnloadingTags = MutableDefault->AllowGameTagUnloading;
+				bShouldAllowUnloadingTags = Default->AllowGameTagUnloading;
 			}
 
 		}
@@ -714,21 +717,7 @@ void UGameplayTagsManager::PopDeferOnGameplayTagTreeChangedBroadcast()
 
 bool UGameplayTagsManager::ShouldImportTagsFromINI() const
 {
-	UGameplayTagsSettings* MutableDefault = GetMutableDefault<UGameplayTagsSettings>();
-
-	// Deprecated path
-	bool ImportFromINI = false;
-	if (GConfig->GetBool(TEXT("GameplayTags"), TEXT("ImportTagsFromConfig"), ImportFromINI, GEngineIni))
-	{
-		if (ImportFromINI)
-		{
-			MutableDefault->ImportTagsFromConfig = ImportFromINI;
-			UE_LOG(LogGameplayTags, Log, TEXT("ImportTagsFromConfig is in a deprecated location, open and save GameplayTag settings to fix"));
-		}
-		return ImportFromINI;
-	}
-
-	return MutableDefault->ImportTagsFromConfig;
+	return GetDefault<UGameplayTagsSettings>()->ImportTagsFromConfig;
 }
 
 bool UGameplayTagsManager::ShouldUnloadTags() const
@@ -783,11 +772,11 @@ void UGameplayTagsManager::ClearShouldDeferGameplayTagTreeRebuilds(bool bRebuild
 
 void UGameplayTagsManager::GetRestrictedTagConfigFiles(TArray<FString>& RestrictedConfigFiles) const
 {
-	UGameplayTagsSettings* MutableDefault = GetMutableDefault<UGameplayTagsSettings>();
+	const UGameplayTagsSettings* Default = GetDefault<UGameplayTagsSettings>();
 
-	if (MutableDefault)
+	if (Default)
 	{
-		for (const FRestrictedConfigInfo& Config : MutableDefault->RestrictedConfigFiles)
+		for (const FRestrictedConfigInfo& Config : Default->RestrictedConfigFiles)
 		{
 			RestrictedConfigFiles.Add(FString::Printf(TEXT("%sTags/%s"), *FPaths::SourceConfigDir(), *Config.RestrictedConfigName));
 		}
@@ -820,11 +809,11 @@ void UGameplayTagsManager::GetRestrictedTagSources(TArray<const FGameplayTagSour
 
 void UGameplayTagsManager::GetOwnersForTagSource(const FString& SourceName, TArray<FString>& OutOwners) const
 {
-	UGameplayTagsSettings* MutableDefault = GetMutableDefault<UGameplayTagsSettings>();
+	const UGameplayTagsSettings* Default = GetDefault<UGameplayTagsSettings>();
 
-	if (MutableDefault)
+	if (Default)
 	{
-		for (const FRestrictedConfigInfo& Config : MutableDefault->RestrictedConfigFiles)
+		for (const FRestrictedConfigInfo& Config : Default->RestrictedConfigFiles)
 		{
 			if (Config.RestrictedConfigName.Equals(SourceName))
 			{
@@ -1002,30 +991,6 @@ UE_AUTORTFM_ALWAYS_OPEN void UGameplayTagsManager::InitializeManager()
 
 	SingletonManager = NewObject<UGameplayTagsManager>(GetTransientPackage(), NAME_None);
 	SingletonManager->AddToRoot();
-
-	UGameplayTagsSettings* MutableDefault = nullptr;
-	{
-		SCOPE_LOG_GAMEPLAYTAGS(TEXT("UGameplayTagsManager::InitializeManager: Load settings"));
-		MutableDefault = GetMutableDefault<UGameplayTagsSettings>();
-	}
-
-	{
-		SCOPE_LOG_GAMEPLAYTAGS(TEXT("UGameplayTagsManager::InitializeManager: Load deprecated"));
-
-		TArray<FString> GameplayTagTablePaths;
-		GConfig->GetArray(TEXT("GameplayTags"), TEXT("+GameplayTagTableList"), GameplayTagTablePaths, GEngineIni);
-
-		// Report deprecation
-		if (GameplayTagTablePaths.Num() > 0)
-		{
-			UE_LOG(LogGameplayTags, Log, TEXT("GameplayTagTableList is in a deprecated location, open and save GameplayTag settings to fix"));
-			for (const FString& DataTable : GameplayTagTablePaths)
-			{
-				MutableDefault->GameplayTagTableList.AddUnique(DataTable);
-			}
-		}
-	}
-
 	SingletonManager->LoadGameplayTagTables(true);
 	SingletonManager->ConstructGameplayTagTree();
 
diff --git a/Engine/Source/Runtime/SlateCore/Public/Styling/StyleColors.h b/Engine/Source/Runtime/SlateCore/Public/Styling/StyleColors.h
--- a/Engine/Source/Runtime/SlateCore/Public/Styling/StyleColors.h
+++ b/Engine/Source/Runtime/SlateCore/Public/Styling/StyleColors.h
@@ -151,12 +151,12 @@ public:
 		return *GetMutableDefault<USlateThemeManager>();
 	}
 
-	const FLinearColor& GetColor(EStyleColor Color)
+	const FLinearColor& GetColor(EStyleColor Color) const
 	{
 		return ActiveColors.StyleColors[static_cast<int32>(Color)];
 	}
 
-	const FGuid& GetCurrentThemeID()
+	const FGuid& GetCurrentThemeID() const
 	{
 		return CurrentThemeId; 
 	}
diff --git a/Engine/Source/Runtime/UMG/Public/Binding/States/WidgetStateSettings.h b/Engine/Source/Runtime/UMG/Public/Binding/States/WidgetStateSettings.h
--- a/Engine/Source/Runtime/UMG/Public/Binding/States/WidgetStateSettings.h
+++ b/Engine/Source/Runtime/UMG/Public/Binding/States/WidgetStateSettings.h
@@ -29,9 +29,9 @@ class UWidgetStateSettings : public UDeveloperSettings
 
 public:
 		
-	static UWidgetStateSettings* Get() 
+	static const UWidgetStateSettings* Get() 
 	{ 
-		return GetMutableDefault<UWidgetStateSettings>();
+		return GetDefault<UWidgetStateSettings>();
 	}
 
 	//~ Begin UObject Interface.
