GameplayTagTest_PerfTest perf improvements

diff --git a/Engine/Source/Runtime/GameplayTags/Private/Tests/GameplayTagTests.cpp b/Engine/Source/Runtime/GameplayTags/Private/Tests/GameplayTagTests.cpp
--- a/Engine/Source/Runtime/GameplayTags/Private/Tests/GameplayTagTests.cpp
+++ b/Engine/Source/Runtime/GameplayTags/Private/Tests/GameplayTagTests.cpp
@@ -24,78 +24,79 @@ public:
 
 	UDataTable* CreateGameplayDataTable()
 	{	
-		TArray<FString> TestTags;
-
-		TestTags.Add(TEXT("Effect.Damage"));
-		TestTags.Add(TEXT("Effect.Damage.Basic"));
-		TestTags.Add(TEXT("Effect.Damage.Type1"));
-		TestTags.Add(TEXT("Effect.Damage.Type2"));
-		TestTags.Add(TEXT("Effect.Damage.Reduce"));
-		TestTags.Add(TEXT("Effect.Damage.Buffable"));
-		TestTags.Add(TEXT("Effect.Damage.Buff"));
-		TestTags.Add(TEXT("Effect.Damage.Physical"));
-		TestTags.Add(TEXT("Effect.Damage.Fire"));
-		TestTags.Add(TEXT("Effect.Damage.Buffed.FireBuff"));
-		TestTags.Add(TEXT("Effect.Damage.Mitigated.Armor"));
-		TestTags.Add(TEXT("Effect.Lifesteal"));
-		TestTags.Add(TEXT("Effect.Shield"));
-		TestTags.Add(TEXT("Effect.Buff"));
-		TestTags.Add(TEXT("Effect.Immune"));
-		TestTags.Add(TEXT("Effect.FireDamage"));
-		TestTags.Add(TEXT("Effect.Shield.Absorb"));
-		TestTags.Add(TEXT("Effect.Protect.Damage"));
-		TestTags.Add(TEXT("Stackable"));
-		TestTags.Add(TEXT("Stack.DiminishingReturns"));
-		TestTags.Add(TEXT("GameplayCue.Burning"));
-		TestTags.Add(TEXT("Expensive.Status.Tag.Type.1"));
-		TestTags.Add(TEXT("Expensive.Status.Tag.Type.2"));
-		TestTags.Add(TEXT("Expensive.Status.Tag.Type.3"));
-		TestTags.Add(TEXT("Expensive.Status.Tag.Type.4"));
-		TestTags.Add(TEXT("Expensive.Status.Tag.Type.5"));
-		TestTags.Add(TEXT("Expensive.Status.Tag.Type.6"));
-		TestTags.Add(TEXT("Expensive.Status.Tag.Type.7"));
-		TestTags.Add(TEXT("Expensive.Status.Tag.Type.8"));
-		TestTags.Add(TEXT("Expensive.Status.Tag.Type.9"));
-		TestTags.Add(TEXT("Expensive.Status.Tag.Type.10"));
-		TestTags.Add(TEXT("Expensive.Status.Tag.Type.11"));
-		TestTags.Add(TEXT("Expensive.Status.Tag.Type.12"));
-		TestTags.Add(TEXT("Expensive.Status.Tag.Type.13"));
-		TestTags.Add(TEXT("Expensive.Status.Tag.Type.14"));
-		TestTags.Add(TEXT("Expensive.Status.Tag.Type.15"));
-		TestTags.Add(TEXT("Expensive.Status.Tag.Type.16"));
-		TestTags.Add(TEXT("Expensive.Status.Tag.Type.17"));
-		TestTags.Add(TEXT("Expensive.Status.Tag.Type.18"));
-		TestTags.Add(TEXT("Expensive.Status.Tag.Type.19"));
-		TestTags.Add(TEXT("Expensive.Status.Tag.Type.20"));
-		TestTags.Add(TEXT("Expensive.Status.Tag.Type.21"));
-		TestTags.Add(TEXT("Expensive.Status.Tag.Type.22"));
-		TestTags.Add(TEXT("Expensive.Status.Tag.Type.23"));
-		TestTags.Add(TEXT("Expensive.Status.Tag.Type.24"));
-		TestTags.Add(TEXT("Expensive.Status.Tag.Type.25"));
-		TestTags.Add(TEXT("Expensive.Status.Tag.Type.26"));
-		TestTags.Add(TEXT("Expensive.Status.Tag.Type.27"));
-		TestTags.Add(TEXT("Expensive.Status.Tag.Type.28"));
-		TestTags.Add(TEXT("Expensive.Status.Tag.Type.29"));
-		TestTags.Add(TEXT("Expensive.Status.Tag.Type.30"));
-		TestTags.Add(TEXT("Expensive.Status.Tag.Type.31"));
-		TestTags.Add(TEXT("Expensive.Status.Tag.Type.32"));
-		TestTags.Add(TEXT("Expensive.Status.Tag.Type.33"));
-		TestTags.Add(TEXT("Expensive.Status.Tag.Type.34"));
-		TestTags.Add(TEXT("Expensive.Status.Tag.Type.35"));
-		TestTags.Add(TEXT("Expensive.Status.Tag.Type.36"));
-		TestTags.Add(TEXT("Expensive.Status.Tag.Type.37"));
-		TestTags.Add(TEXT("Expensive.Status.Tag.Type.38"));
-		TestTags.Add(TEXT("Expensive.Status.Tag.Type.39"));
-		TestTags.Add(TEXT("Expensive.Status.Tag.Type.40"));
+		constexpr FStringView TestTags[] =
+		{
+			TEXTVIEW("Effect.Damage"),
+			TEXTVIEW("Effect.Damage.Basic"),
+			TEXTVIEW("Effect.Damage.Type1"),
+			TEXTVIEW("Effect.Damage.Type2"),
+			TEXTVIEW("Effect.Damage.Reduce"),
+			TEXTVIEW("Effect.Damage.Buffable"),
+			TEXTVIEW("Effect.Damage.Buff"),
+			TEXTVIEW("Effect.Damage.Physical"),
+			TEXTVIEW("Effect.Damage.Fire"),
+			TEXTVIEW("Effect.Damage.Buffed.FireBuff"),
+			TEXTVIEW("Effect.Damage.Mitigated.Armor"),
+			TEXTVIEW("Effect.Lifesteal"),
+			TEXTVIEW("Effect.Shield"),
+			TEXTVIEW("Effect.Buff"),
+			TEXTVIEW("Effect.Immune"),
+			TEXTVIEW("Effect.FireDamage"),
+			TEXTVIEW("Effect.Shield.Absorb"),
+			TEXTVIEW("Effect.Protect.Damage"),
+			TEXTVIEW("Stackable"),
+			TEXTVIEW("Stack.DiminishingReturns"),
+			TEXTVIEW("GameplayCue.Burning"),
+			TEXTVIEW("Expensive.Status.Tag.Type.1"),
+			TEXTVIEW("Expensive.Status.Tag.Type.2"),
+			TEXTVIEW("Expensive.Status.Tag.Type.3"),
+			TEXTVIEW("Expensive.Status.Tag.Type.4"),
+			TEXTVIEW("Expensive.Status.Tag.Type.5"),
+			TEXTVIEW("Expensive.Status.Tag.Type.6"),
+			TEXTVIEW("Expensive.Status.Tag.Type.7"),
+			TEXTVIEW("Expensive.Status.Tag.Type.8"),
+			TEXTVIEW("Expensive.Status.Tag.Type.9"),
+			TEXTVIEW("Expensive.Status.Tag.Type.10"),
+			TEXTVIEW("Expensive.Status.Tag.Type.11"),
+			TEXTVIEW("Expensive.Status.Tag.Type.12"),
+			TEXTVIEW("Expensive.Status.Tag.Type.13"),
+			TEXTVIEW("Expensive.Status.Tag.Type.14"),
+			TEXTVIEW("Expensive.Status.Tag.Type.15"),
+			TEXTVIEW("Expensive.Status.Tag.Type.16"),
+			TEXTVIEW("Expensive.Status.Tag.Type.17"),
+			TEXTVIEW("Expensive.Status.Tag.Type.18"),
+			TEXTVIEW("Expensive.Status.Tag.Type.19"),
+			TEXTVIEW("Expensive.Status.Tag.Type.20"),
+			TEXTVIEW("Expensive.Status.Tag.Type.21"),
+			TEXTVIEW("Expensive.Status.Tag.Type.22"),
+			TEXTVIEW("Expensive.Status.Tag.Type.23"),
+			TEXTVIEW("Expensive.Status.Tag.Type.24"),
+			TEXTVIEW("Expensive.Status.Tag.Type.25"),
+			TEXTVIEW("Expensive.Status.Tag.Type.26"),
+			TEXTVIEW("Expensive.Status.Tag.Type.27"),
+			TEXTVIEW("Expensive.Status.Tag.Type.28"),
+			TEXTVIEW("Expensive.Status.Tag.Type.29"),
+			TEXTVIEW("Expensive.Status.Tag.Type.30"),
+			TEXTVIEW("Expensive.Status.Tag.Type.31"),
+			TEXTVIEW("Expensive.Status.Tag.Type.32"),
+			TEXTVIEW("Expensive.Status.Tag.Type.33"),
+			TEXTVIEW("Expensive.Status.Tag.Type.34"),
+			TEXTVIEW("Expensive.Status.Tag.Type.35"),
+			TEXTVIEW("Expensive.Status.Tag.Type.36"),
+			TEXTVIEW("Expensive.Status.Tag.Type.37"),
+			TEXTVIEW("Expensive.Status.Tag.Type.38"),
+			TEXTVIEW("Expensive.Status.Tag.Type.39"),
+			TEXTVIEW("Expensive.Status.Tag.Type.40"),
+		};
 
 		auto DataTable = NewObject<UDataTable>(GetTransientPackage(), FName(TEXT("TempDataTable")));
 		DataTable->RowStruct = FGameplayTagTableRow::StaticStruct();
 
 
 		FString CSV(TEXT(",Tag,CategoryText,"));
-		for (int32 Count = 0; Count < TestTags.Num(); Count++)
+		for (int32 Count = 0; Count < UE_ARRAY_COUNT(TestTags); Count++)
 		{
-			CSV.Append(FString::Printf(TEXT("\r\n%d,%s"), Count, *TestTags[Count]));
+			CSV.Append(FString::Printf(TEXT("\r\n%d,%.*s"), Count, TestTags[Count].Len(), TestTags[Count].GetData()));
 		}
 
 		DataTable->CreateTableFromCSVString(CSV);
@@ -108,9 +109,9 @@ public:
 		return DataTable;
 	}
 
-	FGameplayTag GetTagForString(const FString& String)
+	FGameplayTag GetTagForString(FStringView String)
 	{
-		return UGameplayTagsManager::Get().RequestGameplayTag(FName(*String));
+		return UGameplayTagsManager::Get().RequestGameplayTag(FName(String));
 	}
 
 	void GameplayTagTest_SimpleTest()
@@ -268,7 +269,7 @@ public:
 				TagContainer.AddTag(CueTag);
 				for (int32 j = 1; j <= 40; j++)
 				{
-					TagContainer.AddTag(GetTagForString(FString::Printf(TEXT("Expensive.Status.Tag.Type.%d"), j)));
+					TagContainer.AddTag(GetTagForString(WriteToString<64>(TEXT("Expensive.Status.Tag.Type."), j)));
 				}
 			}
 		}
