UGameplayTagsManager::AddTagTableRow - Stack-based tag handling to avoid FString/TArray allocations.

#tests All GameplayTag tests pass

diff --git a/Engine/Source/Runtime/GameplayTags/Private/GameplayTagsManager.cpp b/Engine/Source/Runtime/GameplayTags/Private/GameplayTagsManager.cpp
--- a/Engine/Source/Runtime/GameplayTags/Private/GameplayTagsManager.cpp
+++ b/Engine/Source/Runtime/GameplayTags/Private/GameplayTagsManager.cpp
@@ -1032,20 +1032,20 @@ void UGameplayTagsManager::AddTagTableRow(const FGameplayTagTableRow& TagRow, FN
 	}
 
 	// Split the tag text on the "." delimiter to establish tag depth and then insert each tag into the gameplay tag tree
-	// We try to avoid as many FString->FName conversions as possible as they are slow
+	// We try to avoid as many string->FName conversions as possible as they are slow
 	FName OriginalTagName = TagRow.Tag;
-	FString FullTagString = OriginalTagName.ToString();
+	FNameBuilder FullTagString(OriginalTagName);
 
 #if WITH_EDITOR
 	{
 		// In editor builds, validate string
 		// These must get fixed up cooking to work properly
 		FText ErrorText;
-		FString FixedString;
+		FNameBuilder FixedString;
 
 		if (!IsValidGameplayTagString(FullTagString, &ErrorText, &FixedString))
 		{
-			if (FixedString.IsEmpty())
+			if (FixedString.Len() == 0)
 			{
 				// No way to fix it
 				UE_LOG(LogGameplayTags, Error, TEXT("Invalid tag %s from source %s: %s!"), *FullTagString, *SourceName.ToString(), *ErrorText.ToString());
@@ -1054,46 +1054,67 @@ void UGameplayTagsManager::AddTagTableRow(const FGameplayTagTableRow& TagRow, FN
 			else
 			{
 				UE_LOG(LogGameplayTags, Error, TEXT("Invalid tag %s from source %s: %s! Replacing with %s, you may need to modify InvalidTagCharacters"), *FullTagString, *SourceName.ToString(), *ErrorText.ToString(), *FixedString);
-				FullTagString = FixedString;
+				FullTagString.Reset();
+				FullTagString << FixedString;
 				OriginalTagName = FName(*FixedString);
 			}
 		}
 	}
 #endif
 
-	TArray<FString> SubTags;
-	FullTagString.ParseIntoArray(SubTags, TEXT("."), true);
-
-	// We will build this back up as we go
-	FullTagString.Reset();
-
-	int32 NumSubTags = SubTags.Num();
 	bool bHasSeenConflict = false;
 
-	for (int32 SubTagIdx = 0; SubTagIdx < NumSubTags; ++SubTagIdx)
+	constexpr FAsciiSet Period(".");
+
+	const FStringView FullTagView(FullTagString);
+	FStringView Remainder(FullTagView);
+	while (!Remainder.IsEmpty())
 	{
-		bool bIsExplicitTag = (SubTagIdx == (NumSubTags - 1));
-		FName ShortTagName = *SubTags[SubTagIdx];
-		FName FullTagName;
+		// Extract the next component.
+		const FStringView SubTag = FAsciiSet::FindPrefixWithout(Remainder, Period);
+		Remainder.RightChopInline(SubTag.Len());
 
+		// CurrentFullTag is the accumulated tag path for this iteration.
+		const FStringView CurrentFullTag = FullTagView.LeftChop(Remainder.Len());
+
+		// Skip the delimiter, if present.
+		const FStringView Skip = FAsciiSet::FindPrefixWith(Remainder, Period);
+		Remainder.RightChopInline(Skip.Len());
+
+		// Skip any empty sub tags in the hierarchy.
+		if (SubTag.IsEmpty())
+		{
+			continue;
+		}
+
+		FName ShortTagName;
+		if (SubTag.Len() == FullTagView.Len())
+		{
+			// Single-component tag: reuse OriginalTagName to avoid an extra hash.
+			ShortTagName = OriginalTagName;
+		}
+		else
+		{
+			// Hierarchical component.
+			ShortTagName = FName(SubTag);
+		}
+
+		FName FullTagName;
+		const bool bIsExplicitTag = CurrentFullTag.Len() == FullTagView.Len();
 		if (bIsExplicitTag)
 		{
-			// We already know the final name
+			// Final iteration: reuse OriginalTagName to avoid an extra hash.
 			FullTagName = OriginalTagName;
 		}
-		else if (SubTagIdx == 0)
+		else if (CurrentFullTag.Len() == SubTag.Len())
 		{
-			// Full tag is the same as short tag, and start building full tag string
+			// Root component: reuse ShortTagName to avoid an extra hash.
 			FullTagName = ShortTagName;
-			FullTagString = SubTags[SubTagIdx];
 		}
 		else
 		{
-			// Add .Tag and use that as full tag
-			FullTagString += TEXT(".");
-			FullTagString += SubTags[SubTagIdx];
-
-			FullTagName = FName(*FullTagString);
+			// Intermediate component.
+			FullTagName = FName(CurrentFullTag);
 		}
 			
 		TArray< TSharedPtr<FGameplayTagNode> >& ChildTags = CurNode.Get()->GetChildTagNodes();
