GameplayTag: Add a ShowOnlyLeafTag UPROPERTY meta tag, with which it's possible to show only the leaf name in the tag chip instead of the full tag path.

diff --git a/Engine/Plugins/Editor/GameplayTagsEditor/Source/GameplayTagsEditor/Private/SGameplayTagCombo.cpp b/Engine/Plugins/Editor/GameplayTagsEditor/Source/GameplayTagsEditor/Private/SGameplayTagCombo.cpp
--- a/Engine/Plugins/Editor/GameplayTagsEditor/Source/GameplayTagsEditor/Private/SGameplayTagCombo.cpp
+++ b/Engine/Plugins/Editor/GameplayTagsEditor/Source/GameplayTagsEditor/Private/SGameplayTagCombo.cpp
@@ -137,7 +137,12 @@ FText SGameplayTagCombo::GetText() const
 		{
 			return LOCTEXT("GameplayTagCombo_MultipleValues", "Multiple Values");
 		}
-		const FGameplayTag GameplayTag = TagsFromProperty.IsEmpty() ? FGameplayTag() : TagsFromProperty[0]; 
+		const FGameplayTag GameplayTag = TagsFromProperty.IsEmpty() ? FGameplayTag() : TagsFromProperty[0];
+		if (PropertyHandle->HasMetaData("ShowOnlyLeafTag"))
+		{
+			return FText::FromName(GameplayTag.GetTagLeafName());
+		}
+
 		return FText::FromName(GameplayTag.GetTagName());
 	}
 	return FText::FromName(TagAttribute.Get().GetTagName());
diff --git a/Engine/Source/Runtime/GameplayTags/Classes/GameplayTagContainer.h b/Engine/Source/Runtime/GameplayTags/Classes/GameplayTagContainer.h
--- a/Engine/Source/Runtime/GameplayTags/Classes/GameplayTagContainer.h
+++ b/Engine/Source/Runtime/GameplayTags/Classes/GameplayTagContainer.h
@@ -170,6 +170,12 @@ struct FGameplayTag
 	 */
 	GAMEPLAYTAGS_API void ParseParentTags(TArray<FGameplayTag>& UniqueParentTags) const;
 
+	/**
+	 * Parses the tag name and returns the name of the leaf.
+	 * For example, calling this on x.y.z would return the z component.
+	 */
+	GAMEPLAYTAGS_API FName GetTagLeafName() const;
+
 	/** Used so we can have a TMap of this struct */
 	FORCEINLINE friend uint32 GetTypeHash(const FGameplayTag& Tag)
 	{
diff --git a/Engine/Source/Runtime/GameplayTags/Private/GameplayTagContainer.cpp b/Engine/Source/Runtime/GameplayTags/Private/GameplayTagContainer.cpp
--- a/Engine/Source/Runtime/GameplayTags/Private/GameplayTagContainer.cpp
+++ b/Engine/Source/Runtime/GameplayTags/Private/GameplayTagContainer.cpp
@@ -1124,6 +1124,26 @@ void FGameplayTag::ParseParentTags(TArray<FGameplayTag>& UniqueParentTags) const
 	}
 }
 
+FName FGameplayTag::GetTagLeafName() const
+{
+	FName RawTag = GetTagName();
+	if (RawTag.IsNone())
+	{
+		return RawTag;
+	}
+
+	const TStringBuilder<FName::StringBufferSize> TagBuffer(InPlace, RawTag);
+	FStringView TagView = TagBuffer.ToView();
+	const int32 DotIndex = UE::String::FindLastChar(TagView, TEXT('.'));
+	if (DotIndex == INDEX_NONE)
+	{
+		return RawTag;
+	}
+
+	TagView.RightChopInline(DotIndex + 1);
+	return FName(TagView);
+}
+
 bool FGameplayTag::MatchesTag(const FGameplayTag& TagToCheck) const
 {
 	CONDITIONAL_SCOPE_CYCLE_COUNTER(STAT_FGameplayTag_MatchesTag, GEnableGameplayTagDetailedStats);
