Optimization to gameplay tags manager
- Removed unnecessary tree re-creation
- Optimized creation to minimize redudant inserts into the tree (makes insertion ~50% faster)

diff --git a/Engine/Source/Runtime/GameplayTags/Classes/GameplayTagsManager.h b/Engine/Source/Runtime/GameplayTags/Classes/GameplayTagsManager.h
--- a/Engine/Source/Runtime/GameplayTags/Classes/GameplayTagsManager.h
+++ b/Engine/Source/Runtime/GameplayTags/Classes/GameplayTagsManager.h
@@ -944,6 +944,9 @@ private:
 	int32 bDeferBroadcastOnGameplayTagTreeChanged = 0;
 	bool bShouldBroadcastDeferredOnGameplayTagTreeChanged = false;
 
+	/** If true, an action that would require a tree rebuild was performed during initialization **/
+	bool bNeedsTreeRebuildOnDoneAddingGameplayTags = false;
+
 	/** String with outlawed characters inside tags */
 	FString InvalidTagCharacters;
 
diff --git a/Engine/Source/Runtime/GameplayTags/Private/GameplayTagsManager.cpp b/Engine/Source/Runtime/GameplayTags/Private/GameplayTagsManager.cpp
--- a/Engine/Source/Runtime/GameplayTags/Private/GameplayTagsManager.cpp
+++ b/Engine/Source/Runtime/GameplayTags/Private/GameplayTagsManager.cpp
@@ -20,6 +20,8 @@
 #include "Widgets/Notifications/SNotificationList.h"
 #include "Misc/AsciiSet.h"
 #include "Misc/CoreDelegates.h"
+#include "Misc/FileHelper.h"
+#include "Misc/CommandLine.h"
 #include "HAL/IConsoleManager.h"
 #include "NativeGameplayTags.h"
 #include "AutoRTFM/AutoRTFM.h"
@@ -117,6 +119,76 @@ namespace GameplayTagUtil
 			}
 		}
 	}
+
+#if !UE_BUILD_SHIPPING
+	static void GatherGameplayTagStringsRecursive(const FGameplayTagNode& RootNode, TArray<FString>& Out)
+	{
+		Out.Add(RootNode.GetCompleteTagString());
+
+		for (const TSharedPtr<FGameplayTagNode>& ChildNode : RootNode.GetChildTagNodes())
+		{
+			GatherGameplayTagStringsRecursive(*ChildNode, Out);
+		}
+	}
+
+	static void DumpGameplayTagStrings(const FGameplayTagNode& RootNode, const FString& Filename)
+	{
+		TArray<FString> Lines;
+		GatherGameplayTagStringsRecursive(RootNode, Lines);
+		Lines.Sort();
+
+		FString TagDumpFilename = FPaths::ProjectSavedDir() / Filename;
+		if (!FFileHelper::SaveStringArrayToFile(Lines, *TagDumpFilename, FFileHelper::EEncodingOptions::ForceUTF8WithoutBOM))
+		{
+			UE_LOG(LogGameplayTags, Display, TEXT("Wrote Tag Dump: %s"), *TagDumpFilename);
+		}
+	}	
+
+	static void DumpRegisteredSearchPaths(const TMap<FString, FGameplayTagSearchPathInfo>& RegisteredSearchPaths, const FString& Filename)
+	{
+		TArray<FString> Lines;
+		for (const TPair<FString, FGameplayTagSearchPathInfo>& Pair : RegisteredSearchPaths)
+		{
+			Lines.Add(FString::Printf(TEXT("%s bWasSearched:%d bWasAddedToTree:%d"), *Pair.Key, int32(Pair.Value.bWasSearched), int32(Pair.Value.bWasAddedToTree)));
+
+			for (int32 Idx = 0; Idx < Pair.Value.SourcesInPath.Num(); ++Idx)
+			{
+				Lines.Add(FString::Printf(TEXT("%s SourcesInPath[%d]: %s"), *Pair.Key, Idx, *Pair.Value.SourcesInPath[Idx].ToString()));
+			}
+
+			for (int32 Idx = 0; Idx < Pair.Value.TagIniList.Num(); ++Idx)
+			{
+				Lines.Add(FString::Printf(TEXT("%s TagIniList[%d]: %s"), *Pair.Key, Idx, *Pair.Value.TagIniList[Idx]));
+			}
+		}
+
+		Lines.Sort();
+
+		FString DumpFilename = FPaths::ProjectSavedDir() / Filename;
+		if (!FFileHelper::SaveStringArrayToFile(Lines, *DumpFilename, FFileHelper::EEncodingOptions::ForceUTF8WithoutBOM))
+		{
+			UE_LOG(LogGameplayTags, Display, TEXT("Wrote RegisteredSearchPaths Dump: %s"), *DumpFilename);
+		}
+	}
+
+	static void DumpRestrictedGameplayTagSourceNames(const TSet<FName>& RestrictedGameplayTagSourceNames, const FString& Filename)
+	{
+		TArray<FString> Lines;
+
+		for (const FName Name : RestrictedGameplayTagSourceNames)
+		{
+			Lines.Add(Name.ToString());
+		}
+
+		Lines.Sort();
+
+		FString DumpFilename = FPaths::ProjectSavedDir() / Filename;
+		if (!FFileHelper::SaveStringArrayToFile(Lines, *DumpFilename, FFileHelper::EEncodingOptions::ForceUTF8WithoutBOM))
+		{
+			UE_LOG(LogGameplayTags, Display, TEXT("Wrote RestrictedGameplayTagSourceNames Dump: %s"), *DumpFilename);
+		}
+	}
+#endif
 }
 
 
@@ -439,6 +511,7 @@ void UGameplayTagsManager::AddTagsFromAdditionalLooseIniFiles(const TArray<FStri
 void UGameplayTagsManager::ConstructGameplayTagTree()
 {
 	SCOPE_LOG_GAMEPLAYTAGS(TEXT("UGameplayTagsManager::ConstructGameplayTagTree"));
+	UE::TScopeLock Lock(GameplayTagMapCritical);
 	TGuardValue<bool> GuardRebuilding(bIsConstructingGameplayTagTree, true);
 	if (!GameplayRootTag.IsValid())
 	{
@@ -478,12 +551,6 @@ void UGameplayTagsManager::ConstructGameplayTagTree()
 			}
 		}
 
-		// If we didn't load any tables it might be async loading, so load again with a flush
-		if (GameplayTagTables.Num() == 0)
-		{
-			LoadGameplayTagTables(false);
-		}
-		
 		{
 			SCOPE_LOG_GAMEPLAYTAGS(TEXT("UGameplayTagsManager::ConstructGameplayTagTree: Construct from data asset"));
 			for (UDataTable* DataTable : GameplayTagTables)
@@ -991,7 +1058,9 @@ UE_AUTORTFM_ALWAYS_OPEN void UGameplayTagsManager::InitializeManager()
 
 	SingletonManager = NewObject<UGameplayTagsManager>(GetTransientPackage(), NAME_None);
 	SingletonManager->AddToRoot();
-	SingletonManager->LoadGameplayTagTables(true);
+
+	//This is always going to be a synchronous load this early in init, so save some time by not attempting anything async
+	SingletonManager->LoadGameplayTagTables(false);
 	SingletonManager->ConstructGameplayTagTree();
 
 	// Bind to end of engine init to be done adding native tags
@@ -1022,7 +1091,7 @@ void UGameplayTagsManager::PopulateTreeFromDataTable(class UDataTable* InTable)
 void UGameplayTagsManager::AddTagTableRow(const FGameplayTagTableRow& TagRow, FName SourceName, bool bIsRestrictedTag)
 {
 	TSharedPtr<FGameplayTagNode> CurNode = GameplayRootTag;
-	TArray<TSharedPtr<FGameplayTagNode>> AncestorNodes;
+	TArray<TSharedPtr<FGameplayTagNode>, TInlineAllocator<6>> AncestorNodes;
 	bool bAllowNonRestrictedChildren = true;
 
 	const FRestrictedGameplayTagTableRow* RestrictedTagRow = static_cast<const FRestrictedGameplayTagTableRow*>(&TagRow);
@@ -1062,64 +1131,79 @@ void UGameplayTagsManager::AddTagTableRow(const FGameplayTagTableRow& TagRow, FN
 	}
 #endif
 
-	bool bHasSeenConflict = false;
+	struct FRequiredTag
+	{
+		FName ShortTagName;
+		FName FullTagName;
+		bool bIsExplicitTag;
+	};
 
-	constexpr FAsciiSet Period(".");
+	TArray<FRequiredTag, TInlineAllocator<6>> RequiredTags;
 
-	const FStringView FullTagView(FullTagString);
-	FStringView Remainder(FullTagView);
-	while (!Remainder.IsEmpty())
 	{
-		// Extract the next component.
-		const FStringView SubTag = FAsciiSet::FindPrefixWithout(Remainder, Period);
-		Remainder.RightChopInline(SubTag.Len());
-
-		// CurrentFullTag is the accumulated tag path for this iteration.
-		const FStringView CurrentFullTag = FullTagView.LeftChop(Remainder.Len());
+		// don't need to lock if we are constructing the gameplay tag tree as it's already handled in an outer scope
+		UE::TConditionalScopeLock Lock(GameplayTagMapCritical, !bIsConstructingGameplayTagTree);
 
-		// Skip the delimiter, if present.
-		const FStringView Skip = FAsciiSet::FindPrefixWith(Remainder, Period);
-		Remainder.RightChopInline(Skip.Len());
+		constexpr FAsciiSet Period(".");
 
-		// Skip any empty sub tags in the hierarchy.
-		if (SubTag.IsEmpty())
+		// first try and push parents to see how far back we need to go, to be able to early out and not over-process the parent nodes
+		// and avoids unnecessary String -> FName lookup
+		const FStringView FullTagView(FullTagString);
+		FStringView Remainder(FullTagView);
+		while (!Remainder.IsEmpty())
 		{
-			continue;
-		}
+			const FStringView CurrentFullTag = Remainder;
 
-		FName ShortTagName;
-		if (SubTag.Len() == FullTagView.Len())
-		{
-			// Single-component tag: reuse OriginalTagName to avoid an extra hash.
-			ShortTagName = OriginalTagName;
-		}
-		else
-		{
-			// Hierarchical component.
-			ShortTagName = FName(SubTag);
-		}
+			const FStringView SubTag = FAsciiSet::FindSuffixWithout(Remainder, Period);
+			Remainder.LeftChopInline(SubTag.Len());
 
-		FName FullTagName;
-		const bool bIsExplicitTag = CurrentFullTag.Len() == FullTagView.Len();
-		if (bIsExplicitTag)
-		{
-			// Final iteration: reuse OriginalTagName to avoid an extra hash.
-			FullTagName = OriginalTagName;
-		}
-		else if (CurrentFullTag.Len() == SubTag.Len())
-		{
-			// Root component: reuse ShortTagName to avoid an extra hash.
-			FullTagName = ShortTagName;
-		}
-		else
-		{
-			// Intermediate component.
-			FullTagName = FName(CurrentFullTag);
-		}
+			// Skip the delimiter, if present.
+			const FStringView Skip = FAsciiSet::FindSuffixWith(Remainder, Period);
+			Remainder.LeftChopInline(Skip.Len());
+
+			// Skip any empty sub tags in the hierarchy.
+			if (SubTag.IsEmpty())
+			{
+				continue;
+			}
+
+			const bool bIsExplicitTag = CurrentFullTag.Len() == FullTagView.Len();
+			const FName FullTagName = bIsExplicitTag ? OriginalTagName : FName(CurrentFullTag);
 			
-		TArray< TSharedPtr<FGameplayTagNode> >& ChildTags = CurNode.Get()->GetChildTagNodes();
-		int32 InsertionIdx = InsertTagIntoNodeArray(ShortTagName, FullTagName, CurNode, ChildTags, SourceName, TagRow.DevComment, bIsExplicitTag, bIsRestrictedTag, bAllowNonRestrictedChildren);
+			// editor builds need everything to be pushed in order to track additional data (SourceName associated with every node)
+#if !WITH_EDITOR
+			if (!bIsExplicitTag) // we assume the most explicit tag isn't in the tree already
+			{
+				TSharedPtr<FGameplayTagNode> *FoundNode = GameplayTagNodeMap.Find(FGameplayTag(FullTagName));
+				if (FoundNode && FoundNode->IsValid())
+				{
+					// this early out leaves us with only missing tags from the tree in RequiredTags and CurNode set to the parent to start pushing into
+					CurNode = *FoundNode;
+					break;
+				}
+			}
+#endif
 
+			RequiredTags.Emplace(
+				FRequiredTag
+				{
+					.ShortTagName = FName(SubTag),
+					.FullTagName = FullTagName,
+					.bIsExplicitTag = bIsExplicitTag
+				}
+			);
+		}
+	}
+
+	bool bHasSeenConflict = false;
+
+	// process backwards as required tags is a LIFO stack
+	while (!RequiredTags.IsEmpty())
+	{
+		const FRequiredTag CurrentTag = RequiredTags.Pop();
+
+		TArray< TSharedPtr<FGameplayTagNode> >& ChildTags = CurNode.Get()->GetChildTagNodes();		
+		int32 InsertionIdx = InsertTagIntoNodeArray(CurrentTag.ShortTagName, CurrentTag.FullTagName, CurNode, ChildTags, SourceName, TagRow.DevComment, CurrentTag.bIsExplicitTag, bIsRestrictedTag, bAllowNonRestrictedChildren);
 		CurNode = ChildTags[InsertionIdx];
 
 		// Tag conflicts only affect the editor so we don't look for them in the game
@@ -1129,7 +1213,7 @@ void UGameplayTagsManager::AddTagTableRow(const FGameplayTagTableRow& TagRow, FN
 			CurNode->bAncestorHasConflict = bHasSeenConflict;
 
 			// If the sources don't match and the tag is explicit and we should've added the tag explicitly here, we have a conflict
-			if (CurNode->GetFirstSourceName() != SourceName && (CurNode->bIsExplicitTag && bIsExplicitTag))
+			if (CurNode->GetFirstSourceName() != SourceName && (CurNode->bIsExplicitTag && CurrentTag.bIsExplicitTag))
 			{
 				// mark all ancestors as having a bad descendant
 				for (TSharedPtr<FGameplayTagNode> CurAncestorNode : AncestorNodes)
@@ -1212,6 +1296,10 @@ void UGameplayTagsManager::HandleGameplayTagTreeChanged(bool bRecreateTree)
 			BroadcastOnGameplayTagTreeChanged();
 		}
 	}
+	else if (bRecreateTree)
+	{
+		bNeedsTreeRebuildOnDoneAddingGameplayTags = true;
+	}
 }
 
 UGameplayTagsManager::~UGameplayTagsManager()
@@ -1308,7 +1396,7 @@ int32 UGameplayTagsManager::InsertTagIntoNodeArray(FName Tag, FName FullTag, TSh
 		{
 			// This critical section is to handle an issue where tag requests come from another thread when async loading from a background thread in FGameplayTagContainer::Serialize.
 			// This function is not generically threadsafe.
-			UE::TScopeLock Lock(GameplayTagMapCritical);
+			UE::TConditionalScopeLock Lock(GameplayTagMapCritical, !bIsConstructingGameplayTagTree);
 			GameplayTagNodeMap.Add(GameplayTag, TagNode);
 		}
 	}
@@ -2266,9 +2354,34 @@ void UGameplayTagsManager::DoneAddingNativeTags()
 		OnLastChanceToAddNativeTags().Broadcast();
 		bDoneAddingNativeTags = true;
 
-		// We may add native tags that are needed for redirectors, so reconstruct the GameplayTag tree
-		DestroyGameplayTagTree();
-		ConstructGameplayTagTree();
+		bool bNeedsRebuild = bNeedsTreeRebuildOnDoneAddingGameplayTags;
+		if (!bNeedsRebuild)
+		{
+			for (const TPair<FString, FGameplayTagSearchPathInfo>& Pair : RegisteredSearchPaths)
+			{
+				if (!Pair.Value.bWasSearched || !Pair.Value.bWasAddedToTree)
+				{
+					bNeedsRebuild = true;
+					break;
+				}
+			}
+		}
+
+		if (bNeedsRebuild)
+		{
+			// We may add native tags that are needed for redirectors, so reconstruct the GameplayTag tree
+			DestroyGameplayTagTree();
+			ConstructGameplayTagTree();
+		}
+
+#if !UE_BUILD_SHIPPING
+		if (FParse::Param(FCommandLine::Get(), TEXT("DumpStartupGameplayTagManagerState")))
+		{
+			GameplayTagUtil::DumpGameplayTagStrings(*GameplayRootTag, TEXT("GameplayTagManager/Tags.txt"));
+			GameplayTagUtil::DumpRegisteredSearchPaths(RegisteredSearchPaths, TEXT("GameplayTagManager/RegisteredSearchPaths.txt"));
+			GameplayTagUtil::DumpRestrictedGameplayTagSourceNames(RestrictedGameplayTagSourceNames, TEXT("GameplayTagManager/RestrictedGameplayTagSourceNames.txt"));			
+		}
+#endif
 
 		OnDoneAddingNativeTagsDelegate().Broadcast();
 	}
