IncrementalCook: GameplayTags: Add a cookdependency on the GameplayTags for packages with GameplayTagContainers. GameplayTags are read from ini, but are read as individual ini files from the Config directory of GameplayFeaturePlugins, so to cover the possiblity that a new ini file might be read from a plugin mounted in a future cook, we need a dependency on the data read from those ini rather than just dependencies on the ini files we know about in the current cook.
#rnx
#incrementalcook

diff --git a/Engine/Source/Runtime/GameplayTags/Classes/GameplayTagsManager.h b/Engine/Source/Runtime/GameplayTags/Classes/GameplayTagsManager.h
--- a/Engine/Source/Runtime/GameplayTags/Classes/GameplayTagsManager.h
+++ b/Engine/Source/Runtime/GameplayTags/Classes/GameplayTagsManager.h
@@ -13,6 +13,9 @@
 #include "Templates/UniquePtr.h"
 #include "Misc/ScopeLock.h"
 #include "Misc/TransactionallySafeCriticalSection.h"
+#if WITH_EDITOR
+#include "Hash/Blake3.h"
+#endif
 
 #include "GameplayTagsManager.generated.h"
 
@@ -20,6 +23,11 @@ class UGameplayTagsList;
 struct FStreamableHandle;
 class FNativeGameplayTag;
 
+#if WITH_EDITOR
+namespace UE::Cook { class FCookDependency; }
+namespace UE::Cook { class ICookInfo; }
+#endif
+
 /** Simple struct for a table row in the gameplay tag table and element in the ini list */
 USTRUCT()
 struct FGameplayTagTableRow : public FTableRowBase
@@ -261,6 +269,14 @@ struct FGameplayTagNode
 	FString GetDevComment() const { return DevComment; }
 #endif
 
+#if WITH_EDITOR
+	/**
+	 * Update the hasher with a deterministic hash of the data on this. Used for e.g. IncrementalCook keys.
+	 * Does not include data from this node's child or parent nodes.
+	 */
+	GAMEPLAYTAGS_API void Hash(FBlake3& Hasher);
+#endif
+
 private:
 	/** Raw name for this tag at current rank in the tree */
 	FName Tag;
@@ -756,6 +772,14 @@ public:
 	
 	GAMEPLAYTAGS_API bool ShowGameplayTagAsHyperLinkEditor(FString TagName);
 
+	/**
+	 * Used for incremental cooking. Create an FCookDependency that reports tags that have been read from ini.
+	 * Packages that pass this dependency to AddCookLoadDependency or AddCookSaveDependency in their OnCookEvent or
+	 * (if Ar.IsCooking()) Serialize function will be invalidated and recooked by the incremental cook whenever those
+	 * tags change.
+	 */
+	GAMEPLAYTAGS_API UE::Cook::FCookDependency CreateCookDependency();
+
 	/** Implementation of console command GameplayTags.DumpSources */
 	void DumpSources(FOutputDevice& Out) const;
 #endif //WITH_EDITOR
@@ -894,6 +918,10 @@ private:
 	/** Call after modifying the tag tree nodes, this will either call the full editor refresh or a limited game refresh */
 	GAMEPLAYTAGS_API void HandleGameplayTagTreeChanged(bool bRecreateTree);
 
+#if WITH_EDITOR
+	void UpdateIncrementalCookHash(UE::Cook::ICookInfo& CookInfo);
+#endif
+
 	// Tag Sources
 	///////////////////////////////////////////////////////
 
@@ -960,6 +988,8 @@ private:
 
 	TSet<FGuid> EditorRefreshGameplayTagTreeSuspendTokens;
 	bool bEditorRefreshGameplayTagTreeRequestedDuringSuspend = false;
+
+	FBlake3Hash IncrementalCookHash;
 #endif //if WITH_EDITOR
 
 	/** Sorted list of nodes, used for network replication */
@@ -975,4 +1005,6 @@ private:
 
 	GAMEPLAYTAGS_API const static FName NAME_Categories;
 	GAMEPLAYTAGS_API const static FName NAME_GameplayTagFilter;
+
+	friend class UGameplayTagsManagerIncrementalCookFunctions;
 };
diff --git a/Engine/Source/Runtime/GameplayTags/Private/GameplayTagContainer.cpp b/Engine/Source/Runtime/GameplayTags/Private/GameplayTagContainer.cpp
--- a/Engine/Source/Runtime/GameplayTags/Private/GameplayTagContainer.cpp
+++ b/Engine/Source/Runtime/GameplayTags/Private/GameplayTagContainer.cpp
@@ -3,6 +3,7 @@
 #include "GameplayTagContainer.h"
 #include "HAL/IConsoleManager.h"
 #include "UObject/CoreNet.h"
+#include "UObject/ObjectSaveContext.h"
 #include "UObject/UnrealType.h"
 #include "Engine/PackageMapClient.h"
 #include "UObject/Package.h"
@@ -875,6 +876,18 @@ bool FGameplayTagContainer::Serialize(FStructuredArchive::FSlot Slot)
 		{
 			UnderlyingArchive.MarkSearchableName(FGameplayTag::StaticStruct(), Tag.TagName);
 		}
+
+#if WITH_EDITOR
+		if (UnderlyingArchive.IsCooking())
+		{
+			check(UnderlyingArchive.GetSavePackageData());
+			FObjectSavePackageSerializeContext& SaveContext = UnderlyingArchive.GetSavePackageData()->SavePackageContext;
+			if (SaveContext.IsHarvestingCookDependencies())
+			{
+				SaveContext.AddCookLoadDependency(UGameplayTagsManager::Get().CreateCookDependency());
+			}
+		}
+#endif
 	}
 
 	return true;
@@ -1558,11 +1571,26 @@ void FGameplayTag::PostSerialize(const FArchive& Ar)
 		UGameplayTagsManager::Get().SingleGameplayTagLoaded(*this, Ar.GetSerializedProperty());
 	}
 
-	if (Ar.IsSaving() && IsValid())
+	if (Ar.IsSaving())
 	{
-		// This marks the saved name for later searching
-		Ar.MarkSearchableName(FGameplayTag::StaticStruct(), TagName);
+		if (IsValid())
+		{
+			// This marks the saved name for later searching
+			Ar.MarkSearchableName(FGameplayTag::StaticStruct(), TagName);
+		}
+#if WITH_EDITOR
+		if (Ar.IsCooking())
+		{
+			check(const_cast<FArchive&>(Ar).GetSavePackageData());
+			FObjectSavePackageSerializeContext& SaveContext =
+				const_cast<FArchive&>(Ar).GetSavePackageData()->SavePackageContext;
+			if (SaveContext.IsHarvestingCookDependencies())
+			{
+				SaveContext.AddCookLoadDependency(UGameplayTagsManager::Get().CreateCookDependency());
+			}
+		}
 	}
+#endif
 }
 
 bool FGameplayTag::ImportTextItem(const TCHAR*& Buffer, int32 PortFlags, UObject* Parent, FOutputDevice* ErrorText)
diff --git a/Engine/Source/Runtime/GameplayTags/Private/GameplayTagRedirectors.cpp b/Engine/Source/Runtime/GameplayTags/Private/GameplayTagRedirectors.cpp
--- a/Engine/Source/Runtime/GameplayTags/Private/GameplayTagRedirectors.cpp
+++ b/Engine/Source/Runtime/GameplayTags/Private/GameplayTagRedirectors.cpp
@@ -4,6 +4,10 @@
 #include "GameplayTagsSettings.h"
 #include "GameplayTagsManager.h"
 #include "Misc/ConfigCacheIni.h"
+#if WITH_EDITOR
+#include "Algo/Sort.h"
+#include "Hash/Blake3.h"
+#endif
 
 #include UE_INLINE_GENERATED_CPP_BY_NAME(GameplayTagRedirectors)
 
@@ -148,4 +152,34 @@ void FGameplayTagRedirectors::AddRedirects(const TArray<FGameplayTagRedirect>& R
 			ensureMsgf(ExistingRedirect->GetTagName() == NewTagName, TEXT("Old tag %s is being redirected to more than one tag. Please remove all the redirections except for one. NewTagName:%s ExistingRedirect:%s"), *OldTagName.ToString(), *NewTagName.ToString(), *ExistingRedirect->GetTagName().ToString());
 		}
 	}
-}
\ No newline at end of file
+}
+
+#if WITH_EDITOR
+void FGameplayTagRedirectors::Hash(FBlake3& Hasher)
+{
+	TArray<FName> SortedKeys;
+	TagRedirects.GenerateKeyArray(SortedKeys);
+	Algo::Sort(SortedKeys, FNameLexicalLess());
+
+	FNameBuilder Builder;
+	for (FName Key : SortedKeys)
+	{
+		FGameplayTag* Value = TagRedirects.Find(Key);
+		check(Value);
+		Builder.Reset();
+		Builder << Key;
+		for (TCHAR& C : Builder)
+		{
+			C = FChar::ToLower(C);
+		}
+		Hasher.Update(*Builder, Builder.Len() * sizeof(**Builder));
+		Builder.Reset();
+		Builder << Value->GetTagName();
+		for (TCHAR& C : Builder)
+		{
+			C = FChar::ToLower(C);
+		}
+		Hasher.Update(*Builder, Builder.Len() * sizeof(**Builder));
+	}
+}
+#endif
diff --git a/Engine/Source/Runtime/GameplayTags/Private/GameplayTagsManager.cpp b/Engine/Source/Runtime/GameplayTags/Private/GameplayTagsManager.cpp
--- a/Engine/Source/Runtime/GameplayTags/Private/GameplayTagsManager.cpp
+++ b/Engine/Source/Runtime/GameplayTags/Private/GameplayTagsManager.cpp
@@ -29,6 +29,9 @@
 #if WITH_EDITOR
 #include "SourceControlHelpers.h"
 #include "ISourceControlModule.h"
+#include "Algo/Sort.h"
+#include "Cooker/CookDependency.h"
+#include "UObject/ICookInfo.h"
 #include "Editor.h"
 #include "PropertyHandle.h"
 FSimpleMulticastDelegate UGameplayTagsManager::OnEditorRefreshGameplayTagTree;
@@ -1065,6 +1068,13 @@ UE_AUTORTFM_ALWAYS_OPEN void UGameplayTagsManager::InitializeManager()
 
 	// Bind to end of engine init to be done adding native tags
 	FCoreDelegates::OnPostEngineInit.AddUObject(SingletonManager, &UGameplayTagsManager::DoneAddingNativeTags);
+
+#if WITH_EDITOR
+	if (IsRunningCookCommandlet())
+	{
+		UE::Cook::FDelegates::CookStarted.AddUObject(SingletonManager, &UGameplayTagsManager::UpdateIncrementalCookHash);
+	}
+#endif
 }
 
 void UGameplayTagsManager::PopulateTreeFromDataTable(class UDataTable* InTable)
@@ -1913,6 +1923,101 @@ bool UGameplayTagsManager::ShowGameplayTagAsHyperLinkEditor(FString TagName)
 	return false;
 }
 
+class UGameplayTagsManagerIncrementalCookFunctions
+{
+public:
+	static void GetIncrementalCookHash(FCbFieldViewIterator Args, UE::Cook::FCookDependencyContext& Context)
+	{
+		UGameplayTagsManager& Manager = UGameplayTagsManager::Get();
+
+		// IncrementalCookHash is set only once at cook start; if it changes after that in a way that impacts
+		// cooked packages, we will not capture that dependency. We need to ensure instead that all data is
+		// up to date during the call to UpdateIncrementalCookHash at start of cook. This will be true only if
+		// all GameplayFeaturePlugins are properly Registered before CookStarted event. We currently rely on
+		// that for other uses during the cook as well.
+		Context.Update(&Manager.IncrementalCookHash, sizeof(Manager.IncrementalCookHash));
+	}
+};
+
+UE_COOK_DEPENDENCY_FUNCTION(GameplayTagsManager, UGameplayTagsManagerIncrementalCookFunctions::GetIncrementalCookHash);
+
+UE::Cook::FCookDependency UGameplayTagsManager::CreateCookDependency()
+{
+	return UE::Cook::FCookDependency::Function(UE_COOK_DEPENDENCY_FUNCTION_CALL(GameplayTagsManager), FCbFieldIterator());
+}
+
+void UGameplayTagsManager::UpdateIncrementalCookHash(UE::Cook::ICookInfo& CookInfo)
+{
+	UE::TScopeLock Lock(GameplayTagMapCritical);
+
+	// Hash all the data that can effect the bytes or cook errors for a package using GameplayTags
+	FBlake3 Hasher;
+
+	// Redirectors
+	FGameplayTagRedirectors::Get().Hash(Hasher);
+
+	// GameplayTagNodeMap
+	TArray<FGameplayTag> SortedKeys;
+	GameplayTagNodeMap.GenerateKeyArray(SortedKeys);
+	SortedKeys.Sort([](const FGameplayTag& A, const FGameplayTag& B)
+		{
+			return A.GetTagName().LexicalLess(B.GetTagName());
+		});
+	for (FGameplayTag& Key : SortedKeys)
+	{
+		TSharedPtr<FGameplayTagNode>* Value = GameplayTagNodeMap.Find(Key);
+		check(Value);
+		{
+			FNameBuilder Builder;
+			Builder << Key.GetTagName();
+			for (TCHAR& C : Builder)
+			{
+				C = FChar::ToLower(C);
+			}
+			Hasher.Update(*Builder, Builder.Len() * sizeof(**Builder));
+		}
+		(*Value)->Hash(Hasher);
+	}
+
+	IncrementalCookHash = Hasher.Finalize();
+}
+
+void FGameplayTagNode::Hash(FBlake3& Hasher)
+{
+	FNameBuilder Builder;
+	Builder << Tag;
+	for (TCHAR& C : Builder)
+	{
+		C = FChar::ToLower(C);
+	}
+	Hasher.Update(*Builder, Builder.Len() * sizeof(**Builder));
+	TArray<FName> SortedNames = SourceNames;
+	Algo::Sort(SortedNames, FNameLexicalLess());
+	for (FName SourceName : SortedNames)
+	{
+		Builder.Reset();
+		Builder << SourceName;
+		for (TCHAR& C : Builder)
+		{
+			C = FChar::ToLower(C);
+		}
+		Hasher.Update(*Builder, Builder.Len() * sizeof(**Builder));
+	}
+	Hasher.Update(*DevComment, DevComment.Len() * sizeof(**DevComment));
+
+	uint8 Flags = 0;
+	int32 BitCount = 0;
+	Flags |= (bIsRestrictedTag ? 1 : 0) << BitCount++;
+	Flags |= (bAllowNonRestrictedChildren ? 1 : 0) << BitCount++;
+	Flags |= (bIsExplicitTag ? 1 : 0) << BitCount++;
+
+	// These flags are transient and only used for UI display in the interactive editor
+	// Flags |= (bDescendantHasConflict ? 1 : 0) << BitCount++;
+	// Flags |= (bNodeHasConflict ? 1 : 0) << BitCount++;
+	// Flags |= (bAncestorHasConflict ? 1 : 0) << BitCount++;
+	check(BitCount <= sizeof(Flags) * 8);
+	Hasher.Update(&Flags, sizeof(Flags));
+}
 #endif // WITH_EDITOR
 
 const FGameplayTagSource* UGameplayTagsManager::FindTagSource(FName TagSourceName) const
diff --git a/Engine/Source/Runtime/GameplayTags/Public/GameplayTagRedirectors.h b/Engine/Source/Runtime/GameplayTags/Public/GameplayTagRedirectors.h
--- a/Engine/Source/Runtime/GameplayTags/Public/GameplayTagRedirectors.h
+++ b/Engine/Source/Runtime/GameplayTags/Public/GameplayTagRedirectors.h
@@ -10,6 +10,10 @@
 
 #include "GameplayTagRedirectors.generated.h"
 
+#if WITH_EDITOR
+class FBlake3;
+#endif
+
 /** A single redirect from a deleted tag to the new tag that should replace it */
 USTRUCT()
 struct FGameplayTagRedirect
@@ -51,6 +55,11 @@ public:
 	/** Refreshes the redirect map after a config change */
 	void RefreshTagRedirects();
 
+#if WITH_EDITOR
+	/** Update the hasher with a deterministic hash of the existing redirects. Used for e.g. IncrementalCook keys. */
+	void Hash(FBlake3& Hasher);
+#endif
+
 private:
 	FGameplayTagRedirectors();
 
