Add any/all tags exact match query for GameplayTagContainer

#rnx

diff --git a/Engine/Source/Runtime/GameplayTags/Classes/GameplayTagContainer.h b/Engine/Source/Runtime/GameplayTags/Classes/GameplayTagContainer.h
--- a/Engine/Source/Runtime/GameplayTags/Classes/GameplayTagContainer.h
+++ b/Engine/Source/Runtime/GameplayTags/Classes/GameplayTagContainer.h
@@ -768,7 +768,9 @@ enum class EGameplayTagQueryExprType : uint8
 	NoTagsMatch,
 	AnyExprMatch,
 	AllExprMatch,
-	NoExprMatch
+	NoExprMatch,
+	AnyTagsExactMatch,
+	AllTagsExactMatch
 };
 
 namespace EGameplayTagQueryStreamVersion
@@ -925,6 +927,8 @@ public:
 	static GAMEPLAYTAGS_API FGameplayTagQuery MakeQuery_MatchAnyTags(FGameplayTagContainer const& InTags);
 	static GAMEPLAYTAGS_API FGameplayTagQuery MakeQuery_MatchAllTags(FGameplayTagContainer const& InTags);
 	static GAMEPLAYTAGS_API FGameplayTagQuery MakeQuery_MatchNoTags(FGameplayTagContainer const& InTags);
+	static GAMEPLAYTAGS_API FGameplayTagQuery MakeQuery_ExactMatchAnyTags(FGameplayTagContainer const& InTags);
+	static GAMEPLAYTAGS_API FGameplayTagQuery MakeQuery_ExactMatchAllTags(FGameplayTagContainer const& InTags);
 
 	static GAMEPLAYTAGS_API FGameplayTagQuery MakeQuery_MatchTag(FGameplayTag const& InTag);
 
@@ -955,6 +959,18 @@ struct FGameplayTagQueryExpression
 		return *this;
 	}
 
+	FGameplayTagQueryExpression& AnyTagsExactMatch()
+	{
+		ExprType = EGameplayTagQueryExprType::AnyTagsExactMatch;
+		return *this;
+	}
+
+	FGameplayTagQueryExpression& AllTagsExactMatch()
+	{
+		ExprType = EGameplayTagQueryExprType::AllTagsExactMatch;
+		return *this;
+	}
+
 	FGameplayTagQueryExpression& AnyExprMatch()
 	{
 		ExprType = EGameplayTagQueryExprType::AnyExprMatch;
@@ -1014,7 +1030,7 @@ struct FGameplayTagQueryExpression
 	/** Returns true if this expression uses the tag data. */
 	FORCEINLINE bool UsesTagSet() const
 	{
-		return (ExprType == EGameplayTagQueryExprType::AllTagsMatch) || (ExprType == EGameplayTagQueryExprType::AnyTagsMatch) || (ExprType == EGameplayTagQueryExprType::NoTagsMatch);
+		return (ExprType == EGameplayTagQueryExprType::AllTagsMatch) || (ExprType == EGameplayTagQueryExprType::AnyTagsMatch) || (ExprType == EGameplayTagQueryExprType::NoTagsMatch) || (ExprType == EGameplayTagQueryExprType::AnyTagsExactMatch) || (ExprType == EGameplayTagQueryExprType::AllTagsExactMatch);
 	}
 	/** Returns true if this expression uses the expression list data. */
 	FORCEINLINE bool UsesExprSet() const
@@ -1132,6 +1148,32 @@ public:
 #endif  // WITH_EDITOR
 };
 
+UCLASS(BlueprintType, editinlinenew, collapseCategories, meta = (DisplayName = "Any Tags Exact Match"))
+class UEditableGameplayTagQueryExpression_AnyTagsExactMatch : public UEditableGameplayTagQueryExpression
+{
+	GENERATED_BODY()
+public:
+	UPROPERTY(EditDefaultsOnly, Category = Expr)
+	FGameplayTagContainer Tags;
+
+#if WITH_EDITOR
+	virtual void EmitTokens(TArray<uint8>& TokenStream, TArray<FGameplayTag>& TagDictionary, FString* DebugString = nullptr) const override;
+#endif  // WITH_EDITOR
+};
+
+UCLASS(BlueprintType, editinlinenew, collapseCategories, meta = (DisplayName = "All Tags Exact Match"))
+class UEditableGameplayTagQueryExpression_AllTagsExactMatch : public UEditableGameplayTagQueryExpression
+{
+	GENERATED_BODY()
+public:
+	UPROPERTY(EditDefaultsOnly, Category = Expr)
+	FGameplayTagContainer Tags;
+
+#if WITH_EDITOR
+	virtual void EmitTokens(TArray<uint8>& TokenStream, TArray<FGameplayTag>& TagDictionary, FString* DebugString = nullptr) const override;
+#endif  // WITH_EDITOR
+};
+
 UCLASS(BlueprintType, editinlinenew, collapseCategories, meta = (DisplayName = "Any Expressions Match"))
 class UEditableGameplayTagQueryExpression_AnyExprMatch : public UEditableGameplayTagQueryExpression
 {
diff --git a/Engine/Source/Runtime/GameplayTags/Private/GameplayTagContainer.cpp b/Engine/Source/Runtime/GameplayTags/Private/GameplayTagContainer.cpp
--- a/Engine/Source/Runtime/GameplayTags/Private/GameplayTagContainer.cpp
+++ b/Engine/Source/Runtime/GameplayTags/Private/GameplayTagContainer.cpp
@@ -154,6 +154,9 @@ private:
 	bool EvalAnyTagsMatch(FGameplayTagContainer const& Tags, bool bSkip);
 	bool EvalAllTagsMatch(FGameplayTagContainer const& Tags, bool bSkip);
 	bool EvalNoTagsMatch(FGameplayTagContainer const& Tags, bool bSkip);
+	bool EvalAnyTagsExactMatch(FGameplayTagContainer const& Tags, bool bSkip);
+	bool EvalAllTagsExactMatch(FGameplayTagContainer const& Tags, bool bSkip);
+	
 
 	bool EvalAnyExprMatch(FGameplayTagContainer const& Tags, bool bSkip);
 	bool EvalAllExprMatch(FGameplayTagContainer const& Tags, bool bSkip);
@@ -396,6 +399,83 @@ bool FQueryEvaluator::EvalNoTagsMatch(FGameplayTagContainer const& Tags, bool bS
 	return Result;
 }
 
+bool FQueryEvaluator::EvalAnyTagsExactMatch(FGameplayTagContainer const& Tags, bool bSkip)
+{
+	bool bShortCircuit = bSkip;
+	bool Result = false;
+
+	// parse tagset
+	int32 const NumTags = GetToken();
+	if (bReadError)
+	{
+		return false;
+	}
+
+	for (int32 Idx = 0; Idx < NumTags; ++Idx)
+	{
+		int32 const TagIdx = GetToken();
+		if (bReadError)
+		{
+			return false;
+		}
+
+		if (bShortCircuit == false)
+		{
+			FGameplayTag Tag = Query.GetTagFromIndex(TagIdx);
+
+			bool bHasTag = Tags.HasTagExact(Tag);
+
+			if (bHasTag)
+			{
+				// one match is sufficient for a true result!
+				bShortCircuit = true;
+				Result = true;
+			}
+		}
+	}
+
+	return Result;
+}
+
+bool FQueryEvaluator::EvalAllTagsExactMatch(FGameplayTagContainer const& Tags, bool bSkip)
+{
+	bool bShortCircuit = bSkip;
+
+	// assume true until proven otherwise
+	bool Result = true;
+
+	// parse tagset
+	int32 const NumTags = GetToken();
+	if (bReadError)
+	{
+		return false;
+	}
+
+	for (int32 Idx = 0; Idx < NumTags; ++Idx)
+	{
+		int32 const TagIdx = GetToken();
+		if (bReadError)
+		{
+			return false;
+		}
+
+		if (bShortCircuit == false)
+		{
+			FGameplayTag const Tag = Query.GetTagFromIndex(TagIdx);
+			bool const bHasTag = Tags.HasTagExact(Tag);
+
+			if (bHasTag == false)
+			{
+				// one failed match is sufficient for a false result
+				bShortCircuit = true;
+				Result = false;
+			}
+		}
+	}
+
+	return Result;
+}
+
 bool FQueryEvaluator::EvalAnyExprMatch(FGameplayTagContainer const& Tags, bool bSkip)
 {
 	bool bShortCircuit = bSkip;
@@ -505,6 +585,10 @@ bool FQueryEvaluator::EvalExpr(FGameplayTagContainer const& Tags, bool bSkip)
 		return EvalAllTagsMatch(Tags, bSkip);
 	case EGameplayTagQueryExprType::NoTagsMatch:
 		return EvalNoTagsMatch(Tags, bSkip);
+	case EGameplayTagQueryExprType::AnyTagsExactMatch:
+		return EvalAnyTagsExactMatch(Tags, bSkip);
+	case EGameplayTagQueryExprType::AllTagsExactMatch:
+		return EvalAllTagsExactMatch(Tags, bSkip);
 
 	case EGameplayTagQueryExprType::AnyExprMatch:
 		return EvalAnyExprMatch(Tags, bSkip);
@@ -1798,6 +1882,28 @@ FGameplayTagQuery FGameplayTagQuery::MakeQuery_MatchNoTags(FGameplayTagContainer
 		);
 }
 
+// static
+FGameplayTagQuery FGameplayTagQuery::MakeQuery_ExactMatchAnyTags(FGameplayTagContainer const& InTags)
+{
+	return FGameplayTagQuery::BuildQuery
+		(
+		FGameplayTagQueryExpression()
+		.AnyTagsExactMatch()
+		.AddTags(InTags)
+		);
+}
+
+// static
+FGameplayTagQuery FGameplayTagQuery::MakeQuery_ExactMatchAllTags(FGameplayTagContainer const& InTags)
+{
+	return FGameplayTagQuery::BuildQuery
+		(
+		FGameplayTagQueryExpression()
+		.AllTagsExactMatch()
+		.AddTags(InTags)
+		);
+}
+
 // static
 FGameplayTagQuery FGameplayTagQuery::MakeQuery_MatchTag(FGameplayTag const & InTag)
 {
@@ -1867,6 +1973,12 @@ UEditableGameplayTagQueryExpression* FQueryEvaluator::ReadEditableQueryExpr(UObj
 	case EGameplayTagQueryExprType::NoExprMatch:
 		ExprClass = UEditableGameplayTagQueryExpression_NoExprMatch::StaticClass();
 		break;
+	case EGameplayTagQueryExprType::AnyTagsExactMatch:
+		ExprClass = UEditableGameplayTagQueryExpression_AnyTagsExactMatch::StaticClass();
+		break;
+	case EGameplayTagQueryExprType::AllTagsExactMatch:
+		ExprClass = UEditableGameplayTagQueryExpression_AllTagsExactMatch::StaticClass();
+		break;
 	}
 
 	UEditableGameplayTagQueryExpression* NewExpr = nullptr;
@@ -1880,6 +1992,8 @@ UEditableGameplayTagQueryExpression* FQueryEvaluator::ReadEditableQueryExpr(UObj
 			case EGameplayTagQueryExprType::AnyTagsMatch:
 			case EGameplayTagQueryExprType::AllTagsMatch:
 			case EGameplayTagQueryExprType::NoTagsMatch:
+			case EGameplayTagQueryExprType::AnyTagsExactMatch:
+			case EGameplayTagQueryExprType::AllTagsExactMatch:
 				ReadEditableQueryTags(NewExpr);
 				break;
 			case EGameplayTagQueryExprType::AnyExprMatch:
@@ -1910,6 +2024,14 @@ void FQueryEvaluator::ReadEditableQueryTags(UEditableGameplayTagQueryExpression*
 	{
 		Tags = &((UEditableGameplayTagQueryExpression_NoTagsMatch*)EditableQueryExpr)->Tags;
 	}
+	else if (EditableQueryExpr->IsA(UEditableGameplayTagQueryExpression_AnyTagsExactMatch::StaticClass()))
+	{
+		Tags = &((UEditableGameplayTagQueryExpression_AnyTagsExactMatch*)EditableQueryExpr)->Tags;
+	}
+	else if (EditableQueryExpr->IsA(UEditableGameplayTagQueryExpression_AllTagsExactMatch::StaticClass()))
+	{
+		Tags = &((UEditableGameplayTagQueryExpression_AllTagsExactMatch*)EditableQueryExpr)->Tags;
+	}
 	ensure(Tags);
 
 	if (Tags)
@@ -2138,6 +2260,42 @@ void UEditableGameplayTagQueryExpression_NoTagsMatch::EmitTokens(TArray<uint8>&
 	}
 }
 
+void UEditableGameplayTagQueryExpression_AnyTagsExactMatch::EmitTokens(TArray<uint8>& TokenStream, TArray<FGameplayTag>& TagDictionary,
+	FString* DebugString) const
+{
+	TokenStream.Add((uint8)EGameplayTagQueryExprType::AnyTagsExactMatch);
+
+	if (DebugString)
+	{
+		DebugString->Append(TEXT(" ANY_EXACT("));
+	}
+	
+	EmitTagTokens(Tags, TokenStream, TagDictionary, DebugString);
+
+	if (DebugString)
+	{
+		DebugString->Append(TEXT(" )"));
+	}
+}
+
+void UEditableGameplayTagQueryExpression_AllTagsExactMatch::EmitTokens(TArray<uint8>& TokenStream, TArray<FGameplayTag>& TagDictionary,
+	FString* DebugString) const
+{
+	TokenStream.Add((uint8)EGameplayTagQueryExprType::AllTagsExactMatch);
+
+	if (DebugString)
+	{
+		DebugString->Append(TEXT(" ALL_EXACT("));
+	}
+	
+	EmitTagTokens(Tags, TokenStream, TagDictionary, DebugString);
+
+	if (DebugString)
+	{
+		DebugString->Append(TEXT(" )"));
+	}
+}
+
 void UEditableGameplayTagQueryExpression_AnyExprMatch::EmitTokens(TArray<uint8>& TokenStream, TArray<FGameplayTag>& TagDictionary, FString* DebugString) const
 {
 	TokenStream.Add((uint8)EGameplayTagQueryExprType::AnyExprMatch);
@@ -2208,6 +2366,8 @@ void FGameplayTagQueryExpression::EmitTokens(TArray<uint8>& TokenStream, TArray<
 	case EGameplayTagQueryExprType::AnyTagsMatch:
 	case EGameplayTagQueryExprType::AllTagsMatch:
 	case EGameplayTagQueryExprType::NoTagsMatch:
+	case EGameplayTagQueryExprType::AnyTagsExactMatch:
+	case EGameplayTagQueryExprType::AllTagsExactMatch:
 	{
 		// emit tagset
 		uint8 NumTags = (uint8)TagSet.Num();
