Bump max replicated tags in a GameplayTagContainer to 65535.

diff --git a/Engine/Source/Runtime/GameplayTags/Private/GameplayTagContainer.cpp b/Engine/Source/Runtime/GameplayTags/Private/GameplayTagContainer.cpp
--- a/Engine/Source/Runtime/GameplayTags/Private/GameplayTagContainer.cpp
+++ b/Engine/Source/Runtime/GameplayTags/Private/GameplayTagContainer.cpp
@@ -1091,9 +1091,9 @@ bool FGameplayTagContainer::NetSerialize(FArchive& Ar, class UPackageMap* Map, b
 
 	if (Ar.IsSaving())
 	{
-		uint8 NumTags = IntCastChecked<uint8, int32>(GameplayTags.Num());
-		uint8 MaxSize = (1 << NumBitsForContainerSize) - 1;
-		if (!ensureMsgf(NumTags <= MaxSize, TEXT("TagContainer has %d elements when max is %d! Tags: %s"), NumTags, MaxSize, *ToStringSimple()))
+		uint16 NumTags = IntCastChecked<uint16, int32>(GameplayTags.Num());
+		const uint16 MaxSize = (1 << NumBitsForContainerSize) - 1;
+		if (!ensureMsgf(GameplayTags.Num() <= MaxSize, TEXT("TagContainer has %d elements when max is %d! Tags: %s"), GameplayTags.Num(), MaxSize, *ToStringSimple()))
 		{
 			NumTags = MaxSize;
 		}
@@ -1112,12 +1112,12 @@ bool FGameplayTagContainer::NetSerialize(FArchive& Ar, class UPackageMap* Map, b
 	else
 	{
 		// No Common Container tags, just replicate this like normal
-		uint8 NumTags = 0;
+		uint16 NumTags = 0;
 		Ar.SerializeBits(&NumTags, NumBitsForContainerSize);
 
 		GameplayTags.Empty(NumTags);
 		GameplayTags.AddDefaulted(NumTags);
-		for (uint8 idx = 0; idx < NumTags; ++idx)
+		for (int32 idx = 0; idx < NumTags; ++idx)
 		{
 			GameplayTags[idx].NetSerialize_Packed(Ar, Map, bOutSuccess);
 		}
