UE: Add more info to a warning about invalid tags

diff --git a/Engine/Source/Runtime/GameplayTags/Private/GameplayTagsManager.cpp b/Engine/Source/Runtime/GameplayTags/Private/GameplayTagsManager.cpp
--- a/Engine/Source/Runtime/GameplayTags/Private/GameplayTagsManager.cpp
+++ b/Engine/Source/Runtime/GameplayTags/Private/GameplayTagsManager.cpp
@@ -940,7 +940,8 @@ void UGameplayTagsManager::RedirectTagsForContainer(FGameplayTagContainer& Conta
 				{
 					FUObjectSerializeContext* LoadContext = FUObjectThreadContext::Get().GetSerializeContext();
 					UObject* LoadingObject = LoadContext ? LoadContext->SerializedObject : nullptr;
-					UE_ASSET_LOG(LogGameplayTags, Warning, *GetPathNameSafe(LoadingObject), TEXT("Invalid GameplayTag %s found in property %s."), *TagName.ToString(), *GetPathNameSafe(SerializingProperty));
+					UPackage* LoadingPackage = LoadingObject ? LoadingObject->GetPackage() : nullptr;
+					UE_ASSET_LOG(LogGameplayTags, Warning, *GetPathNameSafe(LoadingObject), TEXT("Invalid GameplayTag %s found in property %s. Package: %s"), *TagName.ToString(), *GetPathNameSafe(SerializingProperty), *GetNameSafe(LoadingPackage));
 				}
 			}
 		}
@@ -981,7 +982,8 @@ void UGameplayTagsManager::RedirectSingleGameplayTag(FGameplayTag& Tag, FPropert
 			{
 				FUObjectSerializeContext* LoadContext = FUObjectThreadContext::Get().GetSerializeContext();
 				UObject* LoadingObject = LoadContext ? LoadContext->SerializedObject : nullptr;
-				UE_ASSET_LOG(LogGameplayTags, Warning, *GetPathNameSafe(LoadingObject), TEXT("Invalid GameplayTag %s found in property %s."), *TagName.ToString(), *GetPathNameSafe(SerializingProperty));
+				UPackage* LoadingPackage = LoadingObject ? LoadingObject->GetPackage() : nullptr;
+				UE_ASSET_LOG(LogGameplayTags, Warning, *GetPathNameSafe(LoadingObject), TEXT("Invalid GameplayTag %s found in property %s. Package:%s"), *TagName.ToString(), *GetPathNameSafe(SerializingProperty), *GetNameSafe(LoadingPackage));
 			}			
 		}
 	}
@@ -1017,16 +1019,17 @@ bool UGameplayTagsManager::ImportSingleGameplayTag(FGameplayTag& Tag, FName Impo
 			// These are more elaborate checks to ensure we're actually loading a UObject, and not pasting it, compiling it, or other possible paths into this function.
 			const FUObjectSerializeContext* LoadContext = FUObjectThreadContext::Get().GetSerializeContext();
 			const UObject* LoadingObject = LoadContext ? LoadContext->SerializedObject : nullptr;
+			UPackage* LoadingPackage = LoadingObject ? LoadingObject->GetPackage() : nullptr;
 			if (LoadingObject)
 			{
 				// We need to defer the check until after native gameplay tags are done loading (in case the tag has not yet been defined)
 				CallOrRegister_OnDoneAddingNativeTagsDelegate(FSimpleMulticastDelegate::FDelegate::CreateWeakLambda(this,
-					[this, ImportedTagName, AssetName = GetPathNameSafe(LoadingObject), FullObjectPath = LoadingObject->GetFullName()]()
+					[this, ImportedTagName, AssetName = GetPathNameSafe(LoadingObject), FullObjectPath = LoadingObject->GetFullName(), PackageNameString = GetNameSafe(LoadingPackage)]()
 					{
 						// Verify it again -- it could have been a late-loading native tag
 						if (!ValidateTagCreation(ImportedTagName))
 						{
-							UE_ASSET_LOG(LogGameplayTags, Warning, *AssetName, TEXT("Invalid GameplayTag %s found in object %s."), *ImportedTagName.ToString(), *FullObjectPath);
+							UE_ASSET_LOG(LogGameplayTags, Warning, *AssetName, TEXT("Invalid GameplayTag %s found in object %s. Package: %s"), *ImportedTagName.ToString(), *FullObjectPath, *PackageNameString);
 						}
 					}));
 			}
