- Protect TagRedirects TMap concurrent access.
- RedirectTag returns an instance of FGameplayTag instead of a pointer (pointer could get busted if TagRedirects gets reallocated)

diff --git a/Engine/Source/Runtime/GameplayTags/Private/GameplayTagRedirectors.cpp b/Engine/Source/Runtime/GameplayTags/Private/GameplayTagRedirectors.cpp
--- a/Engine/Source/Runtime/GameplayTags/Private/GameplayTagRedirectors.cpp
+++ b/Engine/Source/Runtime/GameplayTags/Private/GameplayTagRedirectors.cpp
@@ -75,7 +75,10 @@ void FGameplayTagRedirectors::AddRedirectsFromSource(const FGameplayTagSource* S
 
 void FGameplayTagRedirectors::RefreshTagRedirects()
 {
-	TagRedirects.Empty();
+	{
+		FScopeLock Lock(&Mutex);
+		TagRedirects.Empty();
+	}
 
 	// Check settings object
 	const UGameplayTagsSettings* Default = GetDefault<UGameplayTagsSettings>();
@@ -93,18 +96,27 @@ void FGameplayTagRedirectors::RefreshTagRedirects()
 
 }
 
-const FGameplayTag* FGameplayTagRedirectors::RedirectTag(const FName& InTagName) const
+bool FGameplayTagRedirectors::RedirectTag(const FName& InTagName, FGameplayTag& OutTag) const
+{
+	FScopeLock Lock(&Mutex);
+	return RedirectTagThreadUnsafe(InTagName, OutTag);
+}
+
+bool FGameplayTagRedirectors::RedirectTagThreadUnsafe(const FName& InTagName, FGameplayTag& OutTag) const
 {
 	if (const FGameplayTag* NewTag = TagRedirects.Find(InTagName))
 	{
-		return NewTag;
+		OutTag = *NewTag;
+		return true;
 	}
 
-	return nullptr;
+	return false;
 }
 
 void FGameplayTagRedirectors::AddRedirects(const TArray<FGameplayTagRedirect>& Redirects)
 {
+	FScopeLock Lock(&Mutex);
+
 	for (const FGameplayTagRedirect& Redirect : Redirects)
 	{
 		FName OldTagName = Redirect.OldTagName;
@@ -157,6 +169,8 @@ void FGameplayTagRedirectors::AddRedirects(const TArray<FGameplayTagRedirect>& R
 #if WITH_EDITOR
 void FGameplayTagRedirectors::Hash(FBlake3& Hasher)
 {
+	FScopeLock Lock(&Mutex);
+
 	TArray<FName> SortedKeys;
 	TagRedirects.GenerateKeyArray(SortedKeys);
 	Algo::Sort(SortedKeys, FNameLexicalLess());
diff --git a/Engine/Source/Runtime/GameplayTags/Private/GameplayTagsManager.cpp b/Engine/Source/Runtime/GameplayTags/Private/GameplayTagsManager.cpp
--- a/Engine/Source/Runtime/GameplayTags/Private/GameplayTagsManager.cpp
+++ b/Engine/Source/Runtime/GameplayTags/Private/GameplayTagsManager.cpp
@@ -915,38 +915,46 @@ void UGameplayTagsManager::SingleGameplayTagLoaded(FGameplayTag& Tag, FProperty*
 void UGameplayTagsManager::RedirectTagsForContainer(FGameplayTagContainer& Container, FProperty* SerializingProperty) const
 {
 	TArray<FName> NamesToRemove;
-	TArray<const FGameplayTag*> TagsToAdd;
+	TArray<FGameplayTag> TagsToAdd;
 
 	// First populate the NamesToRemove and TagsToAdd sets by finding tags in the container that have redirects
-	for (auto TagIt = Container.CreateConstIterator(); TagIt; ++TagIt)
 	{
-		const FName TagName = TagIt->GetTagName();
-		const FGameplayTag* NewTag = FGameplayTagRedirectors::Get().RedirectTag(TagName);
-		if (NewTag)
+		// Lock FGameplayTagRedirectors to safely iterate over gameplay tags using non-thread-safe functions
+		FGameplayTagRedirectors::Get().Lock();
+		ON_SCOPE_EXIT{ FGameplayTagRedirectors::Get().Unlock(); };
+
+		for (auto TagIt = Container.CreateConstIterator(); TagIt; ++TagIt)
 		{
-			NamesToRemove.Add(TagName);
-			if (NewTag->IsValid())
+			const FName TagName = TagIt->GetTagName();
+
+			FGameplayTag NewTag;
+			if (FGameplayTagRedirectors::Get().RedirectTagThreadUnsafe(TagName, NewTag))
 			{
-				TagsToAdd.Add(NewTag);
+				NamesToRemove.Add(TagName);
+				if (NewTag.IsValid())
+				{
+					TagsToAdd.Add(MoveTemp(NewTag));
+				}
 			}
-		}
-#if WITH_EDITOR
-		else if (SerializingProperty)
-		{
-			// Warn about invalid tags at load time in editor builds, too late to fix it in cooked builds
-			FGameplayTag OldTag = RequestGameplayTag(TagName, false);
-			if (!OldTag.IsValid())
+		#if WITH_EDITOR
+			else if (SerializingProperty)
 			{
-				if (ShouldWarnOnInvalidTags())
+				// Warn about invalid tags at load time in editor builds, too late to fix it in cooked builds
+				FGameplayTag OldTag = RequestGameplayTag(TagName, false);
+				if (!OldTag.IsValid())
 				{
-					FUObjectSerializeContext* LoadContext = FUObjectThreadContext::Get().GetSerializeContext();
-					UObject* LoadingObject = LoadContext ? LoadContext->SerializedObject : nullptr;
-					UPackage* LoadingPackage = LoadingObject ? LoadingObject->GetPackage() : nullptr;
-					UE_ASSET_LOG(LogGameplayTags, Warning, *GetPathNameSafe(LoadingObject), TEXT("Invalid GameplayTag %s found in property %s. Package: %s"), *TagName.ToString(), *GetPathNameSafe(SerializingProperty), *GetNameSafe(LoadingPackage));
+					if (ShouldWarnOnInvalidTags())
+					{
+						FUObjectSerializeContext* LoadContext = FUObjectThreadContext::Get().GetSerializeContext();
+						UObject* LoadingObject = LoadContext ? LoadContext->SerializedObject : nullptr;
+						UPackage* LoadingPackage = LoadingObject ? LoadingObject->GetPackage() : nullptr;
+						UE_ASSET_LOG(LogGameplayTags, Warning, *GetPathNameSafe(LoadingObject), TEXT("Invalid GameplayTag %s found in property %s. Package: %s"), *TagName.ToString(), *GetPathNameSafe(SerializingProperty), *GetNameSafe(LoadingPackage));
+					}
 				}
 			}
+		#endif
 		}
-#endif
+		
 	}
 
 	// Remove all tags from the NamesToRemove set
@@ -956,20 +964,21 @@ void UGameplayTagsManager::RedirectTagsForContainer(FGameplayTagContainer& Conta
 	}
 
 	// Add all tags from the TagsToAdd set
-	for (const FGameplayTag* AddTag : TagsToAdd)
+	for (FGameplayTag& AddTag : TagsToAdd)
 	{
-		Container.AddTag(*AddTag);
+		Container.AddTag(MoveTemp(AddTag));
 	}
 }
 
 void UGameplayTagsManager::RedirectSingleGameplayTag(FGameplayTag& Tag, FProperty* SerializingProperty) const
 {
 	const FName TagName = Tag.GetTagName();
-	if (const FGameplayTag* NewTag = FGameplayTagRedirectors::Get().RedirectTag(TagName))
+	FGameplayTag NewTag;
+	if (FGameplayTagRedirectors::Get().RedirectTag(TagName, NewTag))
 	{
-		if (NewTag->IsValid())
+		if (NewTag.IsValid())
 		{
-			Tag = *NewTag;
+			Tag = NewTag;
 		}
 	}
 #if WITH_EDITOR
@@ -1000,9 +1009,10 @@ bool UGameplayTagsManager::ImportSingleGameplayTag(FGameplayTag& Tag, FName Impo
 	}
 
 	bool bRetVal = false;
-	if (const FGameplayTag* RedirectedTag = FGameplayTagRedirectors::Get().RedirectTag(ImportedTagName))
+	FGameplayTag RedirectedTag;
+	if (FGameplayTagRedirectors::Get().RedirectTag(ImportedTagName, RedirectedTag))
 	{
-		Tag = *RedirectedTag;
+		Tag = RedirectedTag;
 		bRetVal = true;
 	}
 	else if (ValidateTagCreation(ImportedTagName))
@@ -2173,12 +2183,13 @@ FGameplayTag UGameplayTagsManager::RequestGameplayTag(FName TagName, bool ErrorI
 
 	// Check if there are redirects for this tag. If so and the redirected tag is in the node map, return it.
 	// Redirects take priority, even if the tag itself may exist.
-	if (const FGameplayTag* RedirectedTag = FGameplayTagRedirectors::Get().RedirectTag(TagName))
+	FGameplayTag RedirectedTag;
+	if (FGameplayTagRedirectors::Get().RedirectTag(TagName, RedirectedTag))
 	{
 		// Check if the redirected tag exists in the node map
-		if (GameplayTagNodeMap.Contains(*RedirectedTag))
+		if (GameplayTagNodeMap.Contains(RedirectedTag))
 		{
-			return *RedirectedTag;
+			return RedirectedTag;
 		}
 
 		// The tag that was redirected to was not found. Error if that was requested.
@@ -2187,7 +2198,7 @@ FGameplayTag UGameplayTagsManager::RequestGameplayTag(FName TagName, bool ErrorI
 			static TSet<FName> MissingRedirectedTagNames;
 			if (!MissingRedirectedTagNames.Contains(TagName))
 			{
-				const FString RedirectedToName = RedirectedTag->GetTagName().ToString();
+				const FString RedirectedToName = RedirectedTag.GetTagName().ToString();
 				ensureAlwaysMsgf(false, TEXT("Requested Gameplay Tag %s was redirected to %s but %s was not found. Fix or remove the redirect from config."), *TagName.ToString(), *RedirectedToName, *RedirectedToName);
 				MissingRedirectedTagNames.Add(TagName);
 			}
@@ -2866,4 +2877,3 @@ void FGameplayTagNode::ResetNode()
 }
 
 #undef LOCTEXT_NAMESPACE
-
diff --git a/Engine/Source/Runtime/GameplayTags/Public/GameplayTagRedirectors.h b/Engine/Source/Runtime/GameplayTags/Public/GameplayTagRedirectors.h
--- a/Engine/Source/Runtime/GameplayTags/Public/GameplayTagRedirectors.h
+++ b/Engine/Source/Runtime/GameplayTags/Public/GameplayTagRedirectors.h
@@ -46,8 +46,23 @@ class FGameplayTagRedirectors
 public:
 	static FGameplayTagRedirectors& Get();
 
-	/** Sees if the tag name should be redirected to a different tag, returns null if there is no active redirect */
-	const FGameplayTag* RedirectTag(const FName& InTagName) const;
+	/** This class can be accessed by both the game thread and the async loading thread.Lock() / Unlock() must be used when calling unsafe functions. */
+	void Lock() 
+	{
+		Mutex.Lock();
+	}
+
+	/** This class can be accessed by both the game thread and the async loading thread.Lock() / Unlock() must be used when calling unsafe functions. */
+	void Unlock()
+	{
+		Mutex.Unlock();
+	}
+
+	/** Sees if the tag name should be redirected to a different tag, returns false if there is no active redirect */
+	bool RedirectTag(const FName& InTagName, FGameplayTag& OutTag) const;
+
+	/** Sees if the tag name should be redirected to a different tag, returns false if there is no active redirect. Not thread safe. Lock should be called prior to use that function */
+	bool RedirectTagThreadUnsafe(const FName& InTagName, FGameplayTag& OutTag) const;
 
 	/** Adds redirects from the specified source. Faster than refreshing the entire list */
 	void AddRedirectsFromSource(const struct FGameplayTagSource* Source);
@@ -68,4 +83,7 @@ private:
 
 	/** The map of ini-configured tag redirectors */
 	TMap<FName, FGameplayTag> TagRedirects;
+
+	/** TagRedirects can be accessed from multiple threads */
+	mutable FCriticalSection Mutex;
 };
