Replacing FGameplayTagRedirectors::Mutex with FTransactionallySafeCriticalSection to fix a crash when UGameplayTagsManager::RequestGameplayTag() is being called inside of a transaction

#localization none
#tests none

#DisabledValidations nonUnityAndStaticAnalysis

diff --git a/Engine/Source/Runtime/GameplayTags/Private/GameplayTagRedirectors.cpp b/Engine/Source/Runtime/GameplayTags/Private/GameplayTagRedirectors.cpp
--- a/Engine/Source/Runtime/GameplayTags/Private/GameplayTagRedirectors.cpp
+++ b/Engine/Source/Runtime/GameplayTags/Private/GameplayTagRedirectors.cpp
@@ -76,7 +76,7 @@ void FGameplayTagRedirectors::AddRedirectsFromSource(const FGameplayTagSource* S
 void FGameplayTagRedirectors::RefreshTagRedirects()
 {
 	{
-		FScopeLock Lock(&Mutex);
+		UE::TScopeLock Lock(Mutex);
 		TagRedirects.Empty();
 	}
 
@@ -98,7 +98,7 @@ void FGameplayTagRedirectors::RefreshTagRedirects()
 
 bool FGameplayTagRedirectors::RedirectTag(const FName& InTagName, FGameplayTag& OutTag) const
 {
-	FScopeLock Lock(&Mutex);
+	UE::TScopeLock Lock(Mutex);
 	return RedirectTagThreadUnsafe(InTagName, OutTag);
 }
 
@@ -115,7 +115,7 @@ bool FGameplayTagRedirectors::RedirectTagThreadUnsafe(const FName& InTagName, FG
 
 void FGameplayTagRedirectors::AddRedirects(const TArray<FGameplayTagRedirect>& Redirects)
 {
-	FScopeLock Lock(&Mutex);
+	UE::TScopeLock Lock(Mutex);
 
 	for (const FGameplayTagRedirect& Redirect : Redirects)
 	{
@@ -169,7 +169,7 @@ void FGameplayTagRedirectors::AddRedirects(const TArray<FGameplayTagRedirect>& R
 #if WITH_EDITOR
 void FGameplayTagRedirectors::Hash(FBlake3& Hasher)
 {
-	FScopeLock Lock(&Mutex);
+	UE::TScopeLock Lock(Mutex);
 
 	TArray<FName> SortedKeys;
 	TagRedirects.GenerateKeyArray(SortedKeys);
diff --git a/Engine/Source/Runtime/GameplayTags/Public/GameplayTagRedirectors.h b/Engine/Source/Runtime/GameplayTags/Public/GameplayTagRedirectors.h
--- a/Engine/Source/Runtime/GameplayTags/Public/GameplayTagRedirectors.h
+++ b/Engine/Source/Runtime/GameplayTags/Public/GameplayTagRedirectors.h
@@ -6,6 +6,7 @@
 #include "UObject/ObjectMacros.h"
 #include "UObject/Object.h"
 #include "UObject/SoftObjectPath.h"
+#include "Misc/TransactionallySafeCriticalSection.h"
 #include "GameplayTagContainer.h"
 
 #include "GameplayTagRedirectors.generated.h"
@@ -85,5 +86,5 @@ private:
 	TMap<FName, FGameplayTag> TagRedirects;
 
 	/** TagRedirects can be accessed from multiple threads */
-	mutable FCriticalSection Mutex;
+	mutable FTransactionallySafeCriticalSection Mutex;
 };
