Remove the lock inside FGameplayTagRedirectors and use the existing lock in the tag manager instead.
This fixes some editor deadlocks and should also be slightly faster.

diff --git a/Engine/Source/Runtime/GameplayTags/Classes/GameplayTagsManager.h b/Engine/Source/Runtime/GameplayTags/Classes/GameplayTagsManager.h
--- a/Engine/Source/Runtime/GameplayTags/Classes/GameplayTagsManager.h
+++ b/Engine/Source/Runtime/GameplayTags/Classes/GameplayTagsManager.h
@@ -979,7 +979,7 @@ private:
 	FString InvalidTagCharacters;
 
 	// This critical section is to handle an issue where tag requests come from another thread when async loading from a background thread in FGameplayTagContainer::Serialize.
-	// This class is not generically threadsafe.
+	// This class is not generically threadsafe, but this should be locked by any operation that could update something read by a background thread.
 	mutable FTransactionallySafeCriticalSection GameplayTagMapCritical;
 
 #if WITH_EDITOR
diff --git a/Engine/Source/Runtime/GameplayTags/Private/GameplayTagRedirectors.cpp b/Engine/Source/Runtime/GameplayTags/Private/GameplayTagRedirectors.cpp
--- a/Engine/Source/Runtime/GameplayTags/Private/GameplayTagRedirectors.cpp
+++ b/Engine/Source/Runtime/GameplayTags/Private/GameplayTagRedirectors.cpp
@@ -75,10 +75,7 @@ void FGameplayTagRedirectors::AddRedirectsFromSource(const FGameplayTagSource* S
 
 void FGameplayTagRedirectors::RefreshTagRedirects()
 {
-	{
-		UE::TScopeLock Lock(Mutex);
-		TagRedirects.Empty();
-	}
+	TagRedirects.Empty();
 
 	// Check settings object
 	const UGameplayTagsSettings* Default = GetDefault<UGameplayTagsSettings>();
@@ -97,12 +94,6 @@ void FGameplayTagRedirectors::RefreshTagRedirects()
 }
 
 bool FGameplayTagRedirectors::RedirectTag(const FName& InTagName, FGameplayTag& OutTag) const
-{
-	UE::TScopeLock Lock(Mutex);
-	return RedirectTagThreadUnsafe(InTagName, OutTag);
-}
-
-bool FGameplayTagRedirectors::RedirectTagThreadUnsafe(const FName& InTagName, FGameplayTag& OutTag) const
 {
 	if (const FGameplayTag* NewTag = TagRedirects.Find(InTagName))
 	{
@@ -115,8 +106,6 @@ bool FGameplayTagRedirectors::RedirectTagThreadUnsafe(const FName& InTagName, FG
 
 void FGameplayTagRedirectors::AddRedirects(const TArray<FGameplayTagRedirect>& Redirects)
 {
-	UE::TScopeLock Lock(Mutex);
-
 	for (const FGameplayTagRedirect& Redirect : Redirects)
 	{
 		FName OldTagName = Redirect.OldTagName;
@@ -169,8 +158,6 @@ void FGameplayTagRedirectors::AddRedirects(const TArray<FGameplayTagRedirect>& R
 #if WITH_EDITOR
 void FGameplayTagRedirectors::Hash(FBlake3& Hasher)
 {
-	UE::TScopeLock Lock(Mutex);
-
 	TArray<FName> SortedKeys;
 	TagRedirects.GenerateKeyArray(SortedKeys);
 	Algo::Sort(SortedKeys, FNameLexicalLess());
diff --git a/Engine/Source/Runtime/GameplayTags/Private/GameplayTagsManager.cpp b/Engine/Source/Runtime/GameplayTags/Private/GameplayTagsManager.cpp
--- a/Engine/Source/Runtime/GameplayTags/Private/GameplayTagsManager.cpp
+++ b/Engine/Source/Runtime/GameplayTags/Private/GameplayTagsManager.cpp
@@ -477,6 +477,8 @@ void UGameplayTagsManager::AddTagsFromAdditionalLooseIniFiles(const TArray<FStri
 			continue;
 		}
 
+		UE::TScopeLock Lock(GameplayTagMapCritical);
+
 		FGameplayTagSource* FoundSource = FindOrAddTagSource(TagSource, EGameplayTagSourceType::TagList);
 
 		UE_CLOG(GAMEPLAYTAGS_VERBOSE, LogGameplayTags, Display, TEXT("Loading Tag File: %s"), *IniFilePath);
@@ -919,16 +921,15 @@ void UGameplayTagsManager::RedirectTagsForContainer(FGameplayTagContainer& Conta
 
 	// First populate the NamesToRemove and TagsToAdd sets by finding tags in the container that have redirects
 	{
-		// Lock FGameplayTagRedirectors to safely iterate over gameplay tags using non-thread-safe functions
-		FGameplayTagRedirectors::Get().Lock();
-		ON_SCOPE_EXIT{ FGameplayTagRedirectors::Get().Unlock(); };
+		// Lock the tag map to safely iterate over gameplay tags using non-thread-safe functions
+		UE::TScopeLock Lock(GameplayTagMapCritical);
 
 		for (auto TagIt = Container.CreateConstIterator(); TagIt; ++TagIt)
 		{
 			const FName TagName = TagIt->GetTagName();
 
 			FGameplayTag NewTag;
-			if (FGameplayTagRedirectors::Get().RedirectTagThreadUnsafe(TagName, NewTag))
+			if (FGameplayTagRedirectors::Get().RedirectTag(TagName, NewTag))
 			{
 				NamesToRemove.Add(TagName);
 				if (NewTag.IsValid())
@@ -972,6 +973,8 @@ void UGameplayTagsManager::RedirectTagsForContainer(FGameplayTagContainer& Conta
 
 void UGameplayTagsManager::RedirectSingleGameplayTag(FGameplayTag& Tag, FProperty* SerializingProperty) const
 {
+	UE::TScopeLock Lock(GameplayTagMapCritical);
+
 	const FName TagName = Tag.GetTagName();
 	FGameplayTag NewTag;
 	if (FGameplayTagRedirectors::Get().RedirectTag(TagName, NewTag))
@@ -1008,6 +1011,7 @@ bool UGameplayTagsManager::ImportSingleGameplayTag(FGameplayTag& Tag, FName Impo
 		return true;
 	}
 
+	UE::TScopeLock Lock(GameplayTagMapCritical);
 	bool bRetVal = false;
 	FGameplayTag RedirectedTag;
 	if (FGameplayTagRedirectors::Get().RedirectTag(ImportedTagName, RedirectedTag))
diff --git a/Engine/Source/Runtime/GameplayTags/Public/GameplayTagRedirectors.h b/Engine/Source/Runtime/GameplayTags/Public/GameplayTagRedirectors.h
--- a/Engine/Source/Runtime/GameplayTags/Public/GameplayTagRedirectors.h
+++ b/Engine/Source/Runtime/GameplayTags/Public/GameplayTagRedirectors.h
@@ -6,7 +6,6 @@
 #include "UObject/ObjectMacros.h"
 #include "UObject/Object.h"
 #include "UObject/SoftObjectPath.h"
-#include "Misc/TransactionallySafeCriticalSection.h"
 #include "GameplayTagContainer.h"
 
 #include "GameplayTagRedirectors.generated.h"
@@ -42,29 +41,15 @@ public:
 	}
 };
 
+/** This class is not thread safe and is protected against issues using the locks in GameplayTagManager */
 class FGameplayTagRedirectors
 {
 public:
 	static FGameplayTagRedirectors& Get();
 
-	/** This class can be accessed by both the game thread and the async loading thread.Lock() / Unlock() must be used when calling unsafe functions. */
-	void Lock() 
-	{
-		Mutex.Lock();
-	}
-
-	/** This class can be accessed by both the game thread and the async loading thread.Lock() / Unlock() must be used when calling unsafe functions. */
-	void Unlock()
-	{
-		Mutex.Unlock();
-	}
-
 	/** Sees if the tag name should be redirected to a different tag, returns false if there is no active redirect */
 	bool RedirectTag(const FName& InTagName, FGameplayTag& OutTag) const;
 
-	/** Sees if the tag name should be redirected to a different tag, returns false if there is no active redirect. Not thread safe. Lock should be called prior to use that function */
-	bool RedirectTagThreadUnsafe(const FName& InTagName, FGameplayTag& OutTag) const;
-
 	/** Adds redirects from the specified source. Faster than refreshing the entire list */
 	void AddRedirectsFromSource(const struct FGameplayTagSource* Source);
 
@@ -84,7 +69,4 @@ private:
 
 	/** The map of ini-configured tag redirectors */
 	TMap<FName, FGameplayTag> TagRedirects;
-
-	/** TagRedirects can be accessed from multiple threads */
-	mutable FTransactionallySafeCriticalSection Mutex;
 };
