Gameplay tag dependencies : remove the global value representing the gameplay tag manager and replace it with a single dependency per gameplay tag.

To track any change in a tag, I need to save :
- its name
- whether it is defined or not in the GameplayTagManager
- whether it is the target of a redirection or not.

The redirection tracking is a bit tricky. When a gameplay tag is loaded, the redirectors are applied right away in the PostSerialize. It means by the time I create dependencies, I lost the original tag name. I would need a way to keep it around but I couldn't find anything satisfying. So instead I do a reverse lookup : I find all the redirectors having the tag name as a target. Since it can find multiple redirectors and only one is relevant to a given tag, this can trigger useless recooks. But even like this it's worth it since the amount of recooks is reduced a lot compare to the existing solution and this is an edge case that shouldn't happen a lot.

- GameplayTagRedirector : new function FindRedirectorsFromTargetTagName used by the dependencies to find all redirectors from a target tag name.
- GameplayTagsManager :
- Replace the existing function dependency with a new one computing dependencies for a single gameplay tag.
- New structure GameplayTagCookDependency used to save/load all the parameters of the function dependency.
- GameplayTagContainer : use the new function dependency in FGameplayTag and FGameplayTagContainer.

#rnx

diff --git a/Engine/Source/Runtime/GameplayTags/Classes/GameplayTagsManager.h b/Engine/Source/Runtime/GameplayTags/Classes/GameplayTagsManager.h
--- a/Engine/Source/Runtime/GameplayTags/Classes/GameplayTagsManager.h
+++ b/Engine/Source/Runtime/GameplayTags/Classes/GameplayTagsManager.h
@@ -19,13 +19,15 @@
 
 #include "GameplayTagsManager.generated.h"
 
+class FCbFieldViewIterator;
+class FNativeGameplayTag;
 class UGameplayTagsList;
 struct FStreamableHandle;
-class FNativeGameplayTag;
 
 #if WITH_EDITOR
 namespace UE::Cook { class FCookDependency; }
 namespace UE::Cook { class ICookInfo; }
+namespace UE::Cook { struct FCookDependencyContext; }
 #endif
 
 /** Simple struct for a table row in the gameplay tag table and element in the ini list */
@@ -773,12 +775,12 @@ public:
 	GAMEPLAYTAGS_API bool ShowGameplayTagAsHyperLinkEditor(FString TagName);
 
 	/**
-	 * Used for incremental cooking. Create an FCookDependency that reports tags that have been read from ini.
+	 * Used for incremental cooking. Create an FCookDependency for a gameplay tag.
 	 * Packages that pass this dependency to AddCookLoadDependency or AddCookSaveDependency in their OnCookEvent or
 	 * (if Ar.IsCooking()) Serialize function will be invalidated and recooked by the incremental cook whenever those
 	 * tags change.
 	 */
-	GAMEPLAYTAGS_API UE::Cook::FCookDependency CreateCookDependency();
+	GAMEPLAYTAGS_API UE::Cook::FCookDependency CreateCookDependency(const FGameplayTag& Tag) const;
 
 	/** Implementation of console command GameplayTags.DumpSources */
 	void DumpSources(FOutputDevice& Out) const;
@@ -919,8 +921,8 @@ private:
 	GAMEPLAYTAGS_API void HandleGameplayTagTreeChanged(bool bRecreateTree);
 
 #if WITH_EDITOR
-	void UpdateIncrementalCookHash(UE::Cook::ICookInfo& CookInfo);
-#endif
+	void ValidateGameplayTagCookDependencies(const FCbFieldViewIterator& Args, UE::Cook::FCookDependencyContext& Context) const;
+#endif //#if WITH_EDITOR
 
 	// Tag Sources
 	///////////////////////////////////////////////////////
diff --git a/Engine/Source/Runtime/GameplayTags/Private/GameplayTagContainer.cpp b/Engine/Source/Runtime/GameplayTags/Private/GameplayTagContainer.cpp
--- a/Engine/Source/Runtime/GameplayTags/Private/GameplayTagContainer.cpp
+++ b/Engine/Source/Runtime/GameplayTags/Private/GameplayTagContainer.cpp
@@ -934,23 +934,36 @@ bool FGameplayTagContainer::Serialize(FStructuredArchive::FSlot Slot)
 
 	if (UnderlyingArchive.IsSaving())
 	{
+
+#if WITH_EDITOR
+		FObjectSavePackageSerializeContext* SaveContext = nullptr;
+		bool bCreateDependencies = false;
+		if (UnderlyingArchive.IsCooking())
+		{
+			FArchiveSavePackageData* SavePackageData = const_cast<FArchive&>(UnderlyingArchive).GetSavePackageData();
+			check(SavePackageData);
+			SaveContext = &SavePackageData->SavePackageContext;
+
+			if (SaveContext->IsHarvestingCookDependencies())
+			{
+				bCreateDependencies = true;
+			}
+		}
+#endif //#if WITH_EDITOR
+
 		// This marks the saved name for later searching
 		for (const FGameplayTag& Tag : GameplayTags)
 		{
 			UnderlyingArchive.MarkSearchableName(FGameplayTag::StaticStruct(), Tag.TagName);
-		}
 
 #if WITH_EDITOR
-		if (UnderlyingArchive.IsCooking())
-		{
-			check(UnderlyingArchive.GetSavePackageData());
-			FObjectSavePackageSerializeContext& SaveContext = UnderlyingArchive.GetSavePackageData()->SavePackageContext;
-			if (SaveContext.IsHarvestingCookDependencies())
+			if (bCreateDependencies)
 			{
-				SaveContext.AddCookLoadDependency(UGameplayTagsManager::Get().CreateCookDependency());
+				SaveContext->AddCookLoadDependency(UGameplayTagsManager::Get().CreateCookDependency(Tag));
 			}
+#endif //#if WITH_EDITOR
+
 		}
-#endif
 	}
 
 	return true;
@@ -1636,15 +1649,16 @@ void FGameplayTag::PostSerialize(const FArchive& Ar)
 #if WITH_EDITOR
 		if (Ar.IsCooking())
 		{
-			check(const_cast<FArchive&>(Ar).GetSavePackageData());
-			FObjectSavePackageSerializeContext& SaveContext =
-				const_cast<FArchive&>(Ar).GetSavePackageData()->SavePackageContext;
+			FArchiveSavePackageData* SavePackageData = const_cast<FArchive&>(Ar).GetSavePackageData();
+			check(SavePackageData);
+			FObjectSavePackageSerializeContext& SaveContext = SavePackageData->SavePackageContext;
+
 			if (SaveContext.IsHarvestingCookDependencies())
 			{
-				SaveContext.AddCookLoadDependency(UGameplayTagsManager::Get().CreateCookDependency());
+				SaveContext.AddCookLoadDependency(UGameplayTagsManager::Get().CreateCookDependency(*this));
 			}
 		}
-#endif
+#endif //#if WITH_EDITOR
 	}
 }
 
diff --git a/Engine/Source/Runtime/GameplayTags/Private/GameplayTagRedirectors.cpp b/Engine/Source/Runtime/GameplayTags/Private/GameplayTagRedirectors.cpp
--- a/Engine/Source/Runtime/GameplayTags/Private/GameplayTagRedirectors.cpp
+++ b/Engine/Source/Runtime/GameplayTags/Private/GameplayTagRedirectors.cpp
@@ -183,4 +183,17 @@ void FGameplayTagRedirectors::Hash(FBlake3& Hasher)
 		Hasher.Update(*Builder, Builder.Len() * sizeof(**Builder));
 	}
 }
+
+bool FGameplayTagRedirectors::FindRedirectorsFromTargetTagName(const FName& TargetTagName, TArray<FName>& SourceTagNames) const
+{
+	for (const TPair<FName, FGameplayTag>& It : TagRedirects) // we should have a reverse map instead of looping
+	{
+		if (It.Value.GetTagName() == TargetTagName)
+		{
+			SourceTagNames.Add(It.Key);
+		}
+	}
+
+	return !SourceTagNames.IsEmpty();
+}
 #endif
diff --git a/Engine/Source/Runtime/GameplayTags/Private/GameplayTagsManager.cpp b/Engine/Source/Runtime/GameplayTags/Private/GameplayTagsManager.cpp
--- a/Engine/Source/Runtime/GameplayTags/Private/GameplayTagsManager.cpp
+++ b/Engine/Source/Runtime/GameplayTags/Private/GameplayTagsManager.cpp
@@ -35,6 +35,8 @@
 #include "UObject/ICookInfo.h"
 #include "Editor.h"
 #include "PropertyHandle.h"
+#include "Serialization/CompactBinaryWriter.h"
+#include "Serialization/CompactBinarySerialization.h"
 FSimpleMulticastDelegate UGameplayTagsManager::OnEditorRefreshGameplayTagTree;
 #endif
 
@@ -194,6 +196,48 @@ namespace GameplayTagUtil
 		}
 	}
 #endif
+
+#if WITH_EDITOR
+	struct FGameplayTagCookDependency
+	{
+	public:
+		FGameplayTagCookDependency()
+			: TagName()
+			, SourceTagNames()
+			, bTagExists(false)
+		{ }
+
+		void Save(FCbWriter& Writer)
+		{
+			Writer << TagName;
+
+			SourceTagNames.Sort(FNameLexicalLess());
+			Writer << SourceTagNames;
+			
+			Writer << bTagExists;
+		}
+
+		bool Load(FCbFieldViewIterator& Args)
+		{
+			FUtf8StringView TagNameAsString = Args.AsString();
+			TagName = FName(TagNameAsString);
+
+			bool bSuccess = LoadFromCompactBinary((++Args), SourceTagNames);
+			if (!bSuccess)
+			{
+				return false;
+			}
+
+			bTagExists = (++Args).AsBool();
+
+			return true;
+		}
+
+		FName TagName;
+		TArray<FName> SourceTagNames; //if the tag was redirected, contains all the source tag names.
+		bool bTagExists;
+	};
+#endif //#if WITH_EDITOR
 }
 
 
@@ -1083,13 +1127,6 @@ UE_AUTORTFM_ALWAYS_OPEN void UGameplayTagsManager::InitializeManager()
 
 	// Bind to end of engine init to be done adding native tags
 	FCoreDelegates::OnPostEngineInit.AddUObject(SingletonManager, &UGameplayTagsManager::DoneAddingNativeTags);
-
-#if WITH_EDITOR
-	if (IsRunningCookCommandlet())
-	{
-		UE::Cook::FDelegates::CookStarted.AddUObject(SingletonManager, &UGameplayTagsManager::UpdateIncrementalCookHash);
-	}
-#endif
 }
 
 void UGameplayTagsManager::PopulateTreeFromDataTable(class UDataTable* InTable)
@@ -1941,60 +1978,107 @@ bool UGameplayTagsManager::ShowGameplayTagAsHyperLinkEditor(FString TagName)
 class UGameplayTagsManagerIncrementalCookFunctions
 {
 public:
-	static void GetIncrementalCookHash(FCbFieldViewIterator Args, UE::Cook::FCookDependencyContext& Context)
+	static void Validate(FCbFieldViewIterator Args, UE::Cook::FCookDependencyContext& Context)
 	{
 		UGameplayTagsManager& Manager = UGameplayTagsManager::Get();
-
-		// IncrementalCookHash is set only once at cook start; if it changes after that in a way that impacts
-		// cooked packages, we will not capture that dependency. We need to ensure instead that all data is
-		// up to date during the call to UpdateIncrementalCookHash at start of cook. This will be true only if
-		// all GameplayFeaturePlugins are properly Registered before CookStarted event. We currently rely on
-		// that for other uses during the cook as well.
-		Context.Update(&Manager.IncrementalCookHash, sizeof(Manager.IncrementalCookHash));
+		Manager.ValidateGameplayTagCookDependencies(Args, Context);
 	}
 };
 
-UE_COOK_DEPENDENCY_FUNCTION(GameplayTagsManager, UGameplayTagsManagerIncrementalCookFunctions::GetIncrementalCookHash);
+UE_COOK_DEPENDENCY_FUNCTION(GameplayTagDependency, UGameplayTagsManagerIncrementalCookFunctions::Validate);
 
-UE::Cook::FCookDependency UGameplayTagsManager::CreateCookDependency()
+UE::Cook::FCookDependency UGameplayTagsManager::CreateCookDependency(const FGameplayTag& Tag) const
 {
-	return UE::Cook::FCookDependency::Function(UE_COOK_DEPENDENCY_FUNCTION_CALL(GameplayTagsManager), FCbFieldIterator());
+	GameplayTagUtil::FGameplayTagCookDependency Dependencies;
+	{
+		UE::TScopeLock Lock(GameplayTagMapCritical);
+
+		// Redirectors
+		bool bRedirectorFound = FGameplayTagRedirectors::Get().FindRedirectorsFromTargetTagName(Tag.GetTagName(), Dependencies.SourceTagNames);
+
+		// Does the tag exist
+		const TSharedPtr<FGameplayTagNode>* Node = GameplayTagNodeMap.Find(Tag);
+		Dependencies.bTagExists = (Node && Node->IsValid()) ? true : false;
+
+		Dependencies.TagName = Tag.GetTagName();
+	}
+
+	FCbWriter Writer;
+	Dependencies.Save(Writer);
+	FCbFieldIterator Fields = Writer.Save();
+
+	return UE::Cook::FCookDependency::Function(UE_COOK_DEPENDENCY_FUNCTION_CALL(GameplayTagDependency), MoveTemp(Fields));
 }
 
-void UGameplayTagsManager::UpdateIncrementalCookHash(UE::Cook::ICookInfo& CookInfo)
+void UGameplayTagsManager::ValidateGameplayTagCookDependencies(const FCbFieldViewIterator& Args, UE::Cook::FCookDependencyContext& Context) const
 {
 	UE::TScopeLock Lock(GameplayTagMapCritical);
 
-	// Hash all the data that can effect the bytes or cook errors for a package using GameplayTags
-	FBlake3 Hasher;
+	FCbFieldViewIterator ArgsField(Args);
+	GameplayTagUtil::FGameplayTagCookDependency Dependency;
+	bool bSuccess = Dependency.Load(ArgsField);
+	if (!bSuccess)
+	{
+		Context.LogInvalidated(FString::Printf(TEXT("Failed to load the parameters.")));
+		return;
+	}
+
+	const FGameplayTag Tag(Dependency.TagName);
+	const TSharedPtr<FGameplayTagNode>* Node = GameplayTagNodeMap.Find(Tag);
+	bool bTagExists = (Node && Node->IsValid()) ? true : false;
+	if (bTagExists != Dependency.bTagExists)
+	{
+		if (bTagExists)
+		{
+			Context.LogInvalidated(FString::Printf(TEXT("Tag %s was missing but does exist now."), *Dependency.TagName.ToString()));
+		}
+		else
+		{
+			Context.LogInvalidated(FString::Printf(TEXT("Tag %s used to exist but is now missing."), *Dependency.TagName.ToString()));
+		}
 
-	// Redirectors
-	FGameplayTagRedirectors::Get().Hash(Hasher);
+		return;
+	}
 
-	// GameplayTagNodeMap
-	TArray<FGameplayTag> SortedKeys;
-	GameplayTagNodeMap.GenerateKeyArray(SortedKeys);
-	SortedKeys.Sort([](const FGameplayTag& A, const FGameplayTag& B)
+	if (Dependency.SourceTagNames.IsEmpty())
+	{
+		// TagName was not redirected in the past. Let's check if it has a redirection now
+		FGameplayTag TargetTag;
+		bool bRedirectorFound = FGameplayTagRedirectors::Get().RedirectTag(Dependency.TagName, TargetTag);
+		if (bRedirectorFound)
 		{
-			return A.GetTagName().LexicalLess(B.GetTagName());
-		});
-	for (FGameplayTag& Key : SortedKeys)
+			Context.LogInvalidated(FString::Printf(TEXT("New redirection for tag %s. Target tag is %s."),
+				*Dependency.TagName.ToString(), *TargetTag.GetTagName().ToString()));
+
+			return;
+		}
+	}
+	else
 	{
-		TSharedPtr<FGameplayTagNode>* Value = GameplayTagNodeMap.Find(Key);
-		check(Value);
+		// TagName was the target of redirectors. SourceTagNames used to be the sources. Let's check if the sources have changed.
+		// I don't use FGameplayTagRedirectors::FindRedirectorsFromTargetTagName here because RedirectTag is faster and because I don't care
+		// about any new redirectors.
+		for (const FName& SourceTagName : Dependency.SourceTagNames)
 		{
-			FNameBuilder Builder;
-			Builder << Key.GetTagName();
-			for (TCHAR& C : Builder)
+			FGameplayTag TargetTag;
+			bool bRedirectorFound = FGameplayTagRedirectors::Get().RedirectTag(SourceTagName, TargetTag);
+			if (!bRedirectorFound)
 			{
-				C = FChar::ToLower(C);
+				Context.LogInvalidated(FString::Printf(TEXT("Tag %s used to be redirected to %s but the redirection can't be found anymore."),
+					*SourceTagName.ToString(), *Dependency.TagName.ToString()));
+				return;
+			}
+
+			FName TargetTagName = TargetTag.GetTagName();
+			if (TargetTagName != Dependency.TagName)
+			{
+				Context.LogInvalidated(FString::Printf(TEXT("Tag %s redirection was modified. It used to be redirected to %s and is now redirected to %s."),
+					*SourceTagName.ToString(), *Dependency.TagName.ToString(), *TargetTagName.ToString()));
+
+				return;
 			}
-			Hasher.Update(*Builder, Builder.Len() * sizeof(**Builder));
 		}
-		(*Value)->Hash(Hasher);
 	}
-
-	IncrementalCookHash = Hasher.Finalize();
 }
 
 void FGameplayTagNode::Hash(FBlake3& Hasher)
diff --git a/Engine/Source/Runtime/GameplayTags/Public/GameplayTagRedirectors.h b/Engine/Source/Runtime/GameplayTags/Public/GameplayTagRedirectors.h
--- a/Engine/Source/Runtime/GameplayTags/Public/GameplayTagRedirectors.h
+++ b/Engine/Source/Runtime/GameplayTags/Public/GameplayTagRedirectors.h
@@ -59,6 +59,13 @@ public:
 #if WITH_EDITOR
 	/** Update the hasher with a deterministic hash of the existing redirects. Used for e.g. IncrementalCook keys. */
 	void Hash(FBlake3& Hasher);
+
+	/** 
+	 * Find all source tag names from a target tag name. Only used in the cook so far. 
+	 * The reverse lookup is done with a loop, there is nothing to make it fast so used only if you have to. Or better, 
+	 * make a reverse lookup structure! 
+	 */
+	bool FindRedirectorsFromTargetTagName(const FName& TargetTagName, TArray<FName>& SourceTagNames) const;
 #endif
 
 private:
