Check for double '..' periods in Gameplay tag strings, and make sure FixedString is valid at the end.
Updated inis with tags using '..' currently.

#changelist validated

diff --git a/Engine/Source/Runtime/GameplayTags/Private/GameplayTagsManager.cpp b/Engine/Source/Runtime/GameplayTags/Private/GameplayTagsManager.cpp
--- a/Engine/Source/Runtime/GameplayTags/Private/GameplayTagsManager.cpp
+++ b/Engine/Source/Runtime/GameplayTags/Private/GameplayTagsManager.cpp
@@ -25,6 +25,7 @@
 #include "HAL/IConsoleManager.h"
 #include "NativeGameplayTags.h"
 #include "AutoRTFM.h"
+#include "String/Join.h"
 
 #if WITH_EDITOR
 #include "SourceControlHelpers.h"
@@ -2374,6 +2375,27 @@ bool IsValidGameplayTagString(const FStringView& TagString, FText* OutError, Fix
 		FixedString = Trimmed;
 	}
 
+	/*
+	FStringView SubString = TagString;
+	while (SubString.Len() > 0)
+	{
+		const FStringView Trimmed = FAsciiSet::TrimPrefixWithout(SubString, Period);
+		if (const FStringView SubTrim = FAsciiSet::TrimPrefixWith(Trimmed.RightChop(1), Period); SubTrim.Len() != Trimmed.Len() - 1 && SubTrim.Len() > 0)
+		{
+			if (OutError)
+			{
+				Errors.Add(LOCTEXT("DoubleDelimiter", "Tag may not have double delimiters ('..')"));
+			}
+			bIsValid = false;
+			TArray<FStringView> FixedSubstrings = {FixedString.LeftChop(Trimmed.Len()), SubTrim};
+			TStringBuilder<2048> Rebuilder(InPlace, UE::String::Join(FixedSubstrings, '.'));
+			FixedString = Rebuilder.ToString();
+		}
+
+		SubString = Trimmed.RightChop(1);
+	}
+	*/
+
 	FText ErrorText;
 	const FText TagContext = LOCTEXT("GameplayTagContext", "Tag");
 	if (!FName::IsValidXName(FixedString, InvalidTagCharacters, OutError ? &ErrorText : nullptr, &TagContext))
