Add TEnumRange support to EGameplayTagSourceType

diff --git a/Engine/Source/Runtime/GameplayTags/Classes/GameplayTagsManager.h b/Engine/Source/Runtime/GameplayTags/Classes/GameplayTagsManager.h
--- a/Engine/Source/Runtime/GameplayTags/Classes/GameplayTagsManager.h
+++ b/Engine/Source/Runtime/GameplayTags/Classes/GameplayTagsManager.h
@@ -4,6 +4,7 @@
 
 #include "AssetRegistry/AssetData.h"
 #include "CoreMinimal.h"
+#include "Misc/EnumRange.h"
 #include "Stats/Stats.h"
 #include "UObject/ObjectMacros.h"
 #include "UObject/Object.h"
@@ -88,6 +89,8 @@ enum class EGameplayTagSourceType : uint8
 	Invalid,			// Not a real source
 };
 
+ENUM_RANGE_BY_COUNT(EGameplayTagSourceType, EGameplayTagSourceType::Invalid);
+
 UENUM()
 enum class EGameplayTagSelectionType : uint8
 {
