GameplayTag dependency hash : convert the tag string to lower case before hashing it.
That's because FName are case insensitive, so FGamplayTag "FOO" and "foo" are the same. But the hashes will be different.

#rnx

diff --git a/Engine/Source/Runtime/GameplayTags/Private/GameplayTagsManager.cpp b/Engine/Source/Runtime/GameplayTags/Private/GameplayTagsManager.cpp
--- a/Engine/Source/Runtime/GameplayTags/Private/GameplayTagsManager.cpp
+++ b/Engine/Source/Runtime/GameplayTags/Private/GameplayTagsManager.cpp
@@ -237,12 +237,14 @@ namespace GameplayTagUtil
 		void Hash(FBlake3& OutHasher)
 		{
 			FString TagNameString = TagName.ToString();
+			TagNameString.ToLowerInline();
 			OutHasher.Update(TagNameString.GetCharArray().GetData(), TagNameString.NumBytesWithoutNull());
 
 			SourceTagNames.Sort(FNameLexicalLess());
 			for (const FName& SourceTagName : SourceTagNames)
 			{
 				FString SourceTagNameString = SourceTagName.ToString();
+				SourceTagNameString.ToLowerInline();
 				OutHasher.Update(SourceTagNameString.GetCharArray().GetData(), SourceTagNameString.NumBytesWithoutNull());
 			}
 
