Fix all types of Native Gameplay Tags to work correctly after initial startup by using CallOrRegister_OnAddNativeTagsDelegate. This can cause synchronization issues when using FastReplication.
Change SetShouldDeferGameplayTagTreeRebuilds and SuspendEditorRefreshGameplayTagTree to use a unified defer system. These may be replaced with a scope structure in the future.
Fix editor widgets to use OnGameplayTagTreeChanged so they happen for the partial rebuild cases used by game feature loading.

diff --git a/Engine/Plugins/Editor/GameplayTagsEditor/Source/GameplayTagsEditor/Private/SGameplayTagPicker.cpp b/Engine/Plugins/Editor/GameplayTagsEditor/Source/GameplayTagsEditor/Private/SGameplayTagPicker.cpp
--- a/Engine/Plugins/Editor/GameplayTagsEditor/Source/GameplayTagsEditor/Private/SGameplayTagPicker.cpp
+++ b/Engine/Plugins/Editor/GameplayTagsEditor/Source/GameplayTagsEditor/Private/SGameplayTagPicker.cpp
@@ -15,6 +15,7 @@
 #include "PropertyHandle.h"
 #include "Widgets/Input/SSearchBox.h"
 #include "GameplayTagsEditorModule.h"
+#include "GameplayTagsModule.h"
 #include "Framework/Application/SlateApplication.h"
 #include "SAddNewGameplayTagWidget.h"
 #include "SRenameGameplayTagDialog.h"
@@ -111,7 +112,7 @@ void SGameplayTagPicker::Construct(const FArguments& InArgs)
 
 	PostUndoRedoDelegateHandle = FEditorDelegates::PostUndoRedo.AddSP(this, &SGameplayTagPicker::OnPostUndoRedo);
 	
-	UGameplayTagsManager::OnEditorRefreshGameplayTagTree.AddSP(this, &SGameplayTagPicker::RefreshOnNextTick);
+	IGameplayTagsModule::OnGameplayTagTreeChanged.AddSP(this, &SGameplayTagPicker::RefreshOnNextTick);
 	UGameplayTagsManager& Manager = UGameplayTagsManager::Get();
 
 	GetFilteredGameplayRootTags(RootFilterString, TagItems);
diff --git a/Engine/Plugins/Editor/GameplayTagsEditor/Source/GameplayTagsEditor/Private/SGameplayTagWidget.cpp b/Engine/Plugins/Editor/GameplayTagsEditor/Source/GameplayTagsEditor/Private/SGameplayTagWidget.cpp
--- a/Engine/Plugins/Editor/GameplayTagsEditor/Source/GameplayTagsEditor/Private/SGameplayTagWidget.cpp
+++ b/Engine/Plugins/Editor/GameplayTagsEditor/Source/GameplayTagsEditor/Private/SGameplayTagWidget.cpp
@@ -15,7 +15,7 @@
 #include "PropertyHandle.h"
 #include "Widgets/Input/SSearchBox.h"
 #include "GameplayTagsEditorModule.h"
-
+#include "GameplayTagsModule.h"
 #include "Framework/Application/SlateApplication.h"
 #include "SAddNewGameplayTagWidget.h"
 #include "SAddNewGameplayTagSourceWidget.h"
@@ -88,7 +88,7 @@ void SGameplayTagWidget::Construct(const FArguments& InArgs, const TArray<FEdita
 
 	bRestrictedTags = InArgs._RestrictedTags;
 
-	UGameplayTagsManager::OnEditorRefreshGameplayTagTree.AddSP(this, &SGameplayTagWidget::RefreshOnNextTick);
+	IGameplayTagsModule::OnGameplayTagTreeChanged.AddSP(this, &SGameplayTagWidget::RefreshOnNextTick);
 	UGameplayTagsManager& Manager = UGameplayTagsManager::Get();
 
 	GetFilteredGameplayRootTags(RootFilterString, TagItems);
diff --git a/Engine/Plugins/Runtime/CommonUI/Source/CommonUI/Public/UITag.h b/Engine/Plugins/Runtime/CommonUI/Source/CommonUI/Public/UITag.h
--- a/Engine/Plugins/Runtime/CommonUI/Source/CommonUI/Public/UITag.h
+++ b/Engine/Plugins/Runtime/CommonUI/Source/CommonUI/Public/UITag.h
@@ -73,15 +73,16 @@ class TTypedTagStaticImpl
 	TTypedTagStaticImpl()
 	{ 
 		LLM_SCOPE(ELLMTag::UI);
-		OnLastChanceToAddNativeTagsHandle = UGameplayTagsManager::OnLastChanceToAddNativeTags().AddLambda([this]()
+		OnLastChanceToAddNativeTagsHandle = UGameplayTagsManager::CallOrRegister_OnAddNativeTagsDelegate(
+			FSimpleMulticastDelegate::FDelegate::CreateLambda([this]()
 			{
 				StaticImpl.RootTag = UGameplayTagsManager::Get().AddNativeGameplayTag(TagT::GetRootTagStr());
-			});
+			}));
 	}
 
 	~TTypedTagStaticImpl()
 	{ 
-		UGameplayTagsManager::OnLastChanceToAddNativeTags().Remove(OnLastChanceToAddNativeTagsHandle);
+		UGameplayTagsManager::UnregisterNativeTagDelegate(OnLastChanceToAddNativeTagsHandle);
 	}
 
 	TagT RootTag;
diff --git a/Engine/Source/Runtime/GameplayTags/Classes/GameplayTagContainer.h b/Engine/Source/Runtime/GameplayTags/Classes/GameplayTagContainer.h
--- a/Engine/Source/Runtime/GameplayTags/Classes/GameplayTagContainer.h
+++ b/Engine/Source/Runtime/GameplayTags/Classes/GameplayTagContainer.h
@@ -677,6 +677,8 @@ struct FGameplayTagNativeAdder
 	virtual GAMEPLAYTAGS_API ~FGameplayTagNativeAdder();
 
 	virtual void AddTags() = 0;
+	
+	FDelegateHandle AddTagsHandle;
 };
 
 /** Helper struct: drop this in another struct to get an embedded create new tag widget. */
diff --git a/Engine/Source/Runtime/GameplayTags/Classes/GameplayTagsManager.h b/Engine/Source/Runtime/GameplayTags/Classes/GameplayTagsManager.h
--- a/Engine/Source/Runtime/GameplayTags/Classes/GameplayTagsManager.h
+++ b/Engine/Source/Runtime/GameplayTags/Classes/GameplayTagsManager.h
@@ -392,7 +392,7 @@ class UGameplayTagsManager : public UObject
 
 	/**
 	 * Registers the given name as a gameplay tag, and tracks that it is being directly referenced from code
-	 * This can only be called during engine initialization, the table needs to be locked down before replication
+	 * This can be called after initial engine startup but can invalidate FastReplication
 	 *
 	 * @param TagName The Name of the tag to add
 	 * @param TagDevComment The developer comment clarifying the usage of the tag
@@ -411,14 +411,27 @@ public:
 	GAMEPLAYTAGS_API void DoneAddingNativeTags();
 
 	/** This is a delegate that is called during initialization/initial loading and signals the last chance to add tags before we are considered to be fully loaded (all tags registered). */
+	UE_DEPRECATED(5.8, "Replace with CallOrRegister_OnAddNativeTagsDelegate or call AddNativeGameplayTag directly, tags can be safely added after startup with default settings")
 	static GAMEPLAYTAGS_API FSimpleMulticastDelegate& OnLastChanceToAddNativeTags();
 
+	/**
+	 * Register a callback for initialization to add additional native tags right before fully initializing gameplay tags.
+	 * Or, if the startup native tags have already been added, execute Delegate before the end of the current module load.
+	 */
+	static GAMEPLAYTAGS_API FDelegateHandle CallOrRegister_OnAddNativeTagsDelegate(const FSimpleMulticastDelegate::FDelegate& Delegate);
+
 	/**
 	 * Register a callback for when native tags are done being added (this is also a safe point to consider that the gameplay tags have fully been initialized).
 	 * Or, if the native tags have already been added (and thus we have registered all valid tags), then execute this Delegate immediately.
 	 * This is useful if your code is potentially executed during load time, and therefore any tags in your block of code could be not-yet-loaded, but possibly valid after being loaded.
 	 */
-	GAMEPLAYTAGS_API FDelegateHandle CallOrRegister_OnDoneAddingNativeTagsDelegate(const FSimpleMulticastDelegate::FDelegate& Delegate) const;
+	static GAMEPLAYTAGS_API FDelegateHandle CallOrRegister_OnDoneAddingNativeTagsDelegate(const FSimpleMulticastDelegate::FDelegate& Delegate);
+
+	/** 
+	 * Unregisters a previously registered callback from either OnAddNativeTagsDelegate or OnDoneAddingNativeTagsDelegate.
+	 * If the delegate has already activated or the handle is empty this will do nothing.
+	 */ 
+	static GAMEPLAYTAGS_API void UnregisterNativeTagDelegate(FDelegateHandle DelegateHandle);
 
 	/**
 	 * Gets a Tag Container containing the supplied tag and all of its parents as explicit tags.
@@ -610,6 +623,12 @@ public:
 	/** Stops suppressing GameplayTag tree rebuilds and (optionally) rebuilds the tree */
 	GAMEPLAYTAGS_API void ClearShouldDeferGameplayTagTreeRebuilds(bool bRebuildTree);
 
+	/** Temporarily disables IGameplayTagsModule::OnGameplayTagTreeChanged callback for changes to the gameplay tag tree. This does not stop it from actually rebuilding the tree */
+	GAMEPLAYTAGS_API void PushDeferOnGameplayTagTreeChangedBroadcast();
+
+	/** Re-enables OnGameplayTagTreeChanged callbacks, there must be one call to this function for every Push. This will execute callbacks immediately if they were deferred */
+	GAMEPLAYTAGS_API void PopDeferOnGameplayTagTreeChangedBroadcast();
+
 	/** Returns the hash of NetworkGameplayTagNodeIndex */
 	uint32 GetNetworkGameplayTagNodeIndexHash() const { VerifyNetworkIndex(); return NetworkGameplayTagNodeIndexHash; }
 
@@ -658,9 +677,6 @@ public:
 	/** Numbers of bits to use for replicating container size. This can be set via config. */
 	int32 NumBitsForContainerSize;
 
-	GAMEPLAYTAGS_API void PushDeferOnGameplayTagTreeChangedBroadcast();
-	GAMEPLAYTAGS_API void PopDeferOnGameplayTagTreeChangedBroadcast();
-
 private:
 	/** Cached number of bits we need to replicate tags. That is, Log2(Number of Tags). Will always be <= 16. */
 	int32 NetIndexTrueBitNum;
@@ -714,7 +730,10 @@ public:
 	/** Returns information about tag. If not found return false */
     GAMEPLAYTAGS_API bool GetTagEditorData(FName TagName, FString& OutComment, TArray<FName>& OutTagSources, bool& bOutIsTagExplicit, bool &bOutIsRestrictedTag, bool &bOutAllowNonRestrictedChildren) const;
 
-	/** This is called after EditorRefreshGameplayTagTree. Useful if you need to do anything editor related when tags are added or removed */
+	/** 
+	 * This is called after EditorRefreshGameplayTagTree. Useful if you need to do anything editor related when the gameplay tag tree is destroyed and recreated.
+	 * IGameplayTagsModule::OnGameplayTagTreeChanged is called first and will happen when any tags are added or removed, not just when the settings are fully refreshed.
+	 */
 	static GAMEPLAYTAGS_API FSimpleMulticastDelegate OnEditorRefreshGameplayTagTree;
 
 	/** Refresh the gameplaytag tree due to an editor change */
@@ -926,6 +945,15 @@ private:
 	/** Call after modifying the tag tree nodes, this will either call the full editor refresh or a limited game refresh */
 	GAMEPLAYTAGS_API void HandleGameplayTagTreeChanged(bool bRecreateTree);
 
+	/** Handles pending tag changes */
+	void HandleModuleObjectsLoaded(FName InModuleName, bool bCanProcessNewlyLoadedObjects);
+	
+	/** Temporarily disables any game or editor rebuilds of the gameplay tag tree, this should only be used when making multiple tag changes at once */
+	void PushDeferGameplayTagTreeRebuilds();
+
+	/** Re-enables game and editor rebuilds of the gameplay tag tree, there must be one call to this function for every Push. This will rebuild the tree immediately if it was deferred */
+	void PopDeferGameplayTagTreeRebuilds();
+
 #if WITH_EDITOR
 
 	GameplayTagUtil::FGameplayTagCookDependency CreateGameplayTagCookDependency(const FGameplayTag& Tag) const;
@@ -977,11 +1005,19 @@ private:
 	/** Used to suppress calls to HandleGameplayTagTreeChanged that would result in a complete rebuild of the GameplayTag tree*/
 	TOptional<bool> ShouldDeferGameplayTagTreeRebuilds;
 
+	/** Used for HandleModuleObjectsLoaded */
+	FDelegateHandle HandleModuleObjectsLoadedHandle;
+
 	/** True if native tags have all been added and flushed */
 	bool bDoneAddingNativeTags;
 
-	int32 bDeferBroadcastOnGameplayTagTreeChanged = 0;
-	bool bShouldBroadcastDeferredOnGameplayTagTreeChanged = false;
+	/** Used to suppress calls to BroadcastOnGameplayTagTreeChanged*/
+	int32 DeferBroadcastOnGameplayTagTreeChanged = 0;
+	bool bBroadcastOnGameplayTagTreeChangedWasDeferred = false;
+
+	/** Used to suppress calls to HandleGameplayTagTreeChanged that would result in a complete rebuild of the GameplayTag tree */
+	int32 DeferGameplayTagTreeRebuilds = 0;
+	bool bGameplayTagTreeRebuildWasDeferred = false;
 
 	/** If true, an action that would require a tree rebuild was performed during initialization **/
 	bool bNeedsTreeRebuildOnDoneAddingGameplayTags = false;
@@ -997,9 +1033,6 @@ private:
 	// Transient editor-only tags to support quick-iteration PIE workflows
 	TSet<FName> TransientEditorTags;
 
-	TSet<FGuid> EditorRefreshGameplayTagTreeSuspendTokens;
-	bool bEditorRefreshGameplayTagTreeRequestedDuringSuspend = false;
-
 	FBlake3Hash IncrementalCookHash;
 #endif //if WITH_EDITOR
 
diff --git a/Engine/Source/Runtime/GameplayTags/Private/GameplayTagContainer.cpp b/Engine/Source/Runtime/GameplayTags/Private/GameplayTagContainer.cpp
--- a/Engine/Source/Runtime/GameplayTags/Private/GameplayTagContainer.cpp
+++ b/Engine/Source/Runtime/GameplayTags/Private/GameplayTagContainer.cpp
@@ -1704,12 +1704,13 @@ void FGameplayTag::FromExportString(const FString& ExportString, int32 PortFlags
 
 FGameplayTagNativeAdder::FGameplayTagNativeAdder()
 {
-	UGameplayTagsManager::OnLastChanceToAddNativeTags().AddRaw(this, &FGameplayTagNativeAdder::AddTags);
+	AddTagsHandle = UGameplayTagsManager::CallOrRegister_OnAddNativeTagsDelegate(
+		FSimpleMulticastDelegate::FDelegate::CreateRaw(this, &FGameplayTagNativeAdder::AddTags));
 }
 
 FGameplayTagNativeAdder::~FGameplayTagNativeAdder()
 {
-	UGameplayTagsManager::OnLastChanceToAddNativeTags().RemoveAll(this);
+	UGameplayTagsManager::UnregisterNativeTagDelegate(AddTagsHandle);
 }
 
 FGameplayTagQuery::FGameplayTagQuery()
diff --git a/Engine/Source/Runtime/GameplayTags/Private/GameplayTagsManager.cpp b/Engine/Source/Runtime/GameplayTags/Private/GameplayTagsManager.cpp
--- a/Engine/Source/Runtime/GameplayTags/Private/GameplayTagsManager.cpp
+++ b/Engine/Source/Runtime/GameplayTags/Private/GameplayTagsManager.cpp
@@ -835,15 +835,33 @@ FGameplayTagNetIndex UGameplayTagsManager::GetNetIndexFromTag(const FGameplayTag
 
 void UGameplayTagsManager::PushDeferOnGameplayTagTreeChangedBroadcast()
 {
-	++bDeferBroadcastOnGameplayTagTreeChanged;
+	++DeferBroadcastOnGameplayTagTreeChanged;
 }
 
 void UGameplayTagsManager::PopDeferOnGameplayTagTreeChangedBroadcast()
 {
-	if (!--bDeferBroadcastOnGameplayTagTreeChanged && bShouldBroadcastDeferredOnGameplayTagTreeChanged)
+	--DeferBroadcastOnGameplayTagTreeChanged;
+	ensure(DeferBroadcastOnGameplayTagTreeChanged >= 0);
+	if (DeferBroadcastOnGameplayTagTreeChanged == 0 && bBroadcastOnGameplayTagTreeChangedWasDeferred)
 	{
-		bShouldBroadcastDeferredOnGameplayTagTreeChanged = false;
-		IGameplayTagsModule::OnGameplayTagTreeChanged.Broadcast();
+		bBroadcastOnGameplayTagTreeChangedWasDeferred = false;
+		BroadcastOnGameplayTagTreeChanged();
+	}
+}
+
+void UGameplayTagsManager::PushDeferGameplayTagTreeRebuilds()
+{
+	++DeferGameplayTagTreeRebuilds;
+}
+
+void UGameplayTagsManager::PopDeferGameplayTagTreeRebuilds()
+{
+	--DeferGameplayTagTreeRebuilds;
+	ensure(DeferGameplayTagTreeRebuilds >= 0);
+	if (DeferGameplayTagTreeRebuilds == 0 && bGameplayTagTreeRebuildWasDeferred)
+	{
+		bGameplayTagTreeRebuildWasDeferred = false;
+		HandleGameplayTagTreeChanged(true);
 	}
 }
 
@@ -889,17 +907,31 @@ void UGameplayTagsManager::ClearShouldUnloadTagsOverride()
 
 void UGameplayTagsManager::SetShouldDeferGameplayTagTreeRebuilds(bool bShouldDeferRebuilds)
 {
+	bool bWasDeferred = ShouldDeferGameplayTagTreeRebuilds.IsSet() && ShouldDeferGameplayTagTreeRebuilds.GetValue();
+
 	ShouldDeferGameplayTagTreeRebuilds = bShouldDeferRebuilds;
+
+	if (!bWasDeferred && bShouldDeferRebuilds)
+	{
+		PushDeferGameplayTagTreeRebuilds();
+	}
+	else if (bWasDeferred && !bShouldDeferRebuilds)
+	{
+		PopDeferGameplayTagTreeRebuilds();
+	}
 }
 
 void UGameplayTagsManager::ClearShouldDeferGameplayTagTreeRebuilds(bool bRebuildTree)
 {
+	bool bWasDeferred = ShouldDeferGameplayTagTreeRebuilds.IsSet() && ShouldDeferGameplayTagTreeRebuilds.GetValue();
+
 	ShouldDeferGameplayTagTreeRebuilds.Reset();
 
-	if (bRebuildTree)
+	if (bWasDeferred)
 	{
-		HandleGameplayTagTreeChanged(true);
-	}
+		// This will do a rebuild if one was queued and there were no other deferrals
+		PopDeferGameplayTagTreeRebuilds();
+	} 
 }
 
 void UGameplayTagsManager::GetRestrictedTagConfigFiles(TArray<FString>& RestrictedConfigFiles) const
@@ -1337,9 +1369,9 @@ void UGameplayTagsManager::MarkChildrenOfNodeConflict(TSharedPtr<FGameplayTagNod
 
 void UGameplayTagsManager::BroadcastOnGameplayTagTreeChanged()
 {
-	if (bDeferBroadcastOnGameplayTagTreeChanged)
+	if (DeferBroadcastOnGameplayTagTreeChanged > 0)
 	{
-		bShouldBroadcastDeferredOnGameplayTagTreeChanged = true;
+		bBroadcastOnGameplayTagTreeChangedWasDeferred = true;
 	}
 	else
 	{
@@ -1352,7 +1384,14 @@ void UGameplayTagsManager::HandleGameplayTagTreeChanged(bool bRecreateTree)
 	// Don't do anything during a reconstruct or before initial native tags are done loading
 	if (!bIsConstructingGameplayTagTree && bDoneAddingNativeTags)
 	{
-		if (bRecreateTree && (!ShouldDeferGameplayTagTreeRebuilds.IsSet() || !ShouldDeferGameplayTagTreeRebuilds.GetValue()))
+		if (bRecreateTree && DeferGameplayTagTreeRebuilds > 0)
+		{
+			// Defer the full rebuild but do the simple one
+			bGameplayTagTreeRebuildWasDeferred = true;
+			bRecreateTree = false;
+		}
+
+		if (bRecreateTree)
 		{
 #if WITH_EDITOR
 			if (GIsEditor)
@@ -1877,9 +1916,11 @@ bool UGameplayTagsManager::GetTagEditorData(FName TagName, FString& OutComment,
 
 void UGameplayTagsManager::EditorRefreshGameplayTagTree()
 {
-	if (!EditorRefreshGameplayTagTreeSuspendTokens.IsEmpty())
+	if (DeferGameplayTagTreeRebuilds > 0)
 	{
-		bEditorRefreshGameplayTagTreeRequestedDuringSuspend = true;
+		// This will call HandleGameplayTagTreeChanged(true) which comes back to here
+		bGameplayTagTreeRebuildWasDeferred = true;
+
 		return;
 	}
 
@@ -1900,17 +1941,12 @@ void UGameplayTagsManager::EditorRefreshGameplayTagTree()
 
 void UGameplayTagsManager::SuspendEditorRefreshGameplayTagTree(FGuid SuspendToken)
 {
-	EditorRefreshGameplayTagTreeSuspendTokens.Add(SuspendToken);
+	PushDeferGameplayTagTreeRebuilds();
 }
 
 void UGameplayTagsManager::ResumeEditorRefreshGameplayTagTree(FGuid SuspendToken)
 {
-	EditorRefreshGameplayTagTreeSuspendTokens.Remove(SuspendToken);
-	if (EditorRefreshGameplayTagTreeSuspendTokens.IsEmpty() && bEditorRefreshGameplayTagTreeRequestedDuringSuspend)
-	{
-		bEditorRefreshGameplayTagTreeRequestedDuringSuspend = false;
-		EditorRefreshGameplayTagTree();
-;	}
+	PopDeferGameplayTagTreeRebuilds();
 }
 
 #endif //if WITH_EDITOR
@@ -2534,22 +2570,19 @@ FGameplayTag UGameplayTagsManager::AddNativeGameplayTag(FName TagName, const FSt
 		return FGameplayTag();
 	}
 
-	// Unsafe to call after done adding
-	if (ensure(!bDoneAddingNativeTags))
-	{
-		FGameplayTag NewTag = FGameplayTag(TagName);
+	FGameplayTag NewTag = FGameplayTag(TagName);
 
-		if (!LegacyNativeTags.Contains(TagName))
-		{
-			LegacyNativeTags.Add(TagName);
-		}
+	if (!LegacyNativeTags.Contains(TagName))
+	{
+		LegacyNativeTags.Add(TagName);
+	}
 
-		AddTagTableRow(FGameplayTagTableRow(TagName, TagDevComment), FGameplayTagSource::GetNativeName());
+	// Add to temporary tree and call nondestructive notify, which won't do anything if we're still in initial startup
+	AddTagTableRow(FGameplayTagTableRow(TagName, TagDevComment), FGameplayTagSource::GetNativeName());
 
-		return NewTag;
-	}
+	HandleGameplayTagTreeChanged(false);
 
-	return FGameplayTag();
+	return NewTag;
 }
 
 void UGameplayTagsManager::AddNativeGameplayTag(FNativeGameplayTag* TagSource)
@@ -2575,9 +2608,27 @@ void UGameplayTagsManager::RemoveNativeGameplayTag(const FNativeGameplayTag* Tag
 	HandleGameplayTagTreeChanged(true);
 }
 
-FDelegateHandle UGameplayTagsManager::CallOrRegister_OnDoneAddingNativeTagsDelegate(const FSimpleMulticastDelegate::FDelegate& Delegate) const
+PRAGMA_DISABLE_DEPRECATION_WARNINGS // Remove when OnLastChanceToAddNativeTags is made private
+
+FDelegateHandle UGameplayTagsManager::CallOrRegister_OnAddNativeTagsDelegate(const FSimpleMulticastDelegate::FDelegate& Delegate)
+{
+	UGameplayTagsManager* Manager = UGameplayTagsManager::GetIfAllocated();
+	if (Manager && Manager->bDoneAddingNativeTags)
+	{
+		if (!Manager->HandleModuleObjectsLoadedHandle.IsValid())
+		{
+			// Need to schedule a module loaded callback as this is called from the FGameplayTagNativeAdder constructor which can't safely call the virtual functions at this scope
+			Manager->HandleModuleObjectsLoadedHandle = FModuleManager::Get().OnProcessLoadedObjectsCallback().AddUObject(Manager, &UGameplayTagsManager::HandleModuleObjectsLoaded);
+		}
+	}
+
+	return OnLastChanceToAddNativeTags().Add(Delegate);
+}
+
+FDelegateHandle UGameplayTagsManager::CallOrRegister_OnDoneAddingNativeTagsDelegate(const FSimpleMulticastDelegate::FDelegate& Delegate)
 {
-	if (bDoneAddingNativeTags)
+	UGameplayTagsManager* Manager = UGameplayTagsManager::GetIfAllocated();
+	if (Manager && Manager->bDoneAddingNativeTags)
 	{
 		Delegate.Execute();
 		return FDelegateHandle{};
@@ -2588,6 +2639,17 @@ FDelegateHandle UGameplayTagsManager::CallOrRegister_OnDoneAddingNativeTagsDeleg
 	}
 }
 
+void UGameplayTagsManager::UnregisterNativeTagDelegate(FDelegateHandle DelegateHandle)
+{
+	if (DelegateHandle.IsValid())
+	{
+		if (!OnLastChanceToAddNativeTags().Remove(DelegateHandle))
+		{
+			OnDoneAddingNativeTagsDelegate().Remove(DelegateHandle);
+		}
+	}
+}
+
 FSimpleMulticastDelegate& UGameplayTagsManager::OnDoneAddingNativeTagsDelegate()
 {
 	static FSimpleMulticastDelegate Delegate;
@@ -2604,12 +2666,17 @@ void UGameplayTagsManager::DoneAddingNativeTags()
 {
 	// Safe to call multiple times, only works the first time, must be called after the engine
 	// is initialized (DoneAddingNativeTags is bound to PostEngineInit to cover anything that's skipped).
+	// CallOrRegister_OnAddNativeTags will execute immediately after the first successful call
 	if (GEngine && !bDoneAddingNativeTags)
 	{
 		TRACE_CPUPROFILER_EVENT_SCOPE(UGameplayTagsManager::DoneAddingNativeTags);
+		
+		FSimpleMulticastDelegate& AddNativeTagsDelegate = OnLastChanceToAddNativeTags();
+
+		UE_CLOG(GAMEPLAYTAGS_VERBOSE, LogGameplayTags, Display, TEXT("UGameplayTagsManager::DoneAddingNativeTags. DelegateIsBound: %d"), (int32)AddNativeTagsDelegate.IsBound());
+		AddNativeTagsDelegate.Broadcast();
+		AddNativeTagsDelegate.Clear();
 
-		UE_CLOG(GAMEPLAYTAGS_VERBOSE, LogGameplayTags, Display, TEXT("UGameplayTagsManager::DoneAddingNativeTags. DelegateIsBound: %d"), (int32)OnLastChanceToAddNativeTags().IsBound());
-		OnLastChanceToAddNativeTags().Broadcast();
 		bDoneAddingNativeTags = true;
 
 		bool bNeedsRebuild = bNeedsTreeRebuildOnDoneAddingGameplayTags;
@@ -2645,6 +2712,22 @@ void UGameplayTagsManager::DoneAddingNativeTags()
 	}
 }
 
+void UGameplayTagsManager::HandleModuleObjectsLoaded(FName InModuleName, bool bCanProcessNewlyLoadedObjects)
+{
+	// Remove the loaded callback
+	FModuleManager::Get().OnProcessLoadedObjectsCallback().Remove(HandleModuleObjectsLoadedHandle);
+	HandleModuleObjectsLoadedHandle.Reset();
+
+	// Handle any pending native tag callbacks from the last load
+	FSimpleMulticastDelegate& AddNativeTagsDelegate = OnLastChanceToAddNativeTags();
+	PushDeferOnGameplayTagTreeChangedBroadcast();
+	AddNativeTagsDelegate.Broadcast();
+	AddNativeTagsDelegate.Clear();
+	PopDeferOnGameplayTagTreeChangedBroadcast();
+}
+
+PRAGMA_ENABLE_DEPRECATION_WARNINGS
+
 FGameplayTagContainer UGameplayTagsManager::RequestGameplayTagParents(const FGameplayTag& GameplayTag) const
 {
 	UE::TScopeLock Lock(GameplayTagMapCritical);
