Fix engine shutdown crash with previous gameplay tag change, it's not safe to unregister from static function-allocated delegates at engine exit

diff --git a/Engine/Source/Runtime/GameplayTags/Private/GameplayTagsManager.cpp b/Engine/Source/Runtime/GameplayTags/Private/GameplayTagsManager.cpp
--- a/Engine/Source/Runtime/GameplayTags/Private/GameplayTagsManager.cpp
+++ b/Engine/Source/Runtime/GameplayTags/Private/GameplayTagsManager.cpp
@@ -2641,7 +2641,7 @@ FDelegateHandle UGameplayTagsManager::CallOrRegister_OnDoneAddingNativeTagsDeleg
 
 void UGameplayTagsManager::UnregisterNativeTagDelegate(FDelegateHandle DelegateHandle)
 {
-	if (DelegateHandle.IsValid())
+	if (DelegateHandle.IsValid() && !IsEngineExitRequested())
 	{
 		if (!OnLastChanceToAddNativeTags().Remove(DelegateHandle))
 		{
