Memory optimization from Robert Millarfor GameplayTagsManager.cpp

diff --git a/Engine/Source/Runtime/GameplayTags/Private/GameplayTagsManager.cpp b/Engine/Source/Runtime/GameplayTags/Private/GameplayTagsManager.cpp
--- a/Engine/Source/Runtime/GameplayTags/Private/GameplayTagsManager.cpp
+++ b/Engine/Source/Runtime/GameplayTags/Private/GameplayTagsManager.cpp
@@ -2,8 +2,11 @@
 
 #include "GameplayTagsManager.h"
 #include "Engine/Engine.h"
+#include "GameplayTagContainer.h"
 #include "HAL/PlatformFileManager.h"
 #include "HAL/FileManager.h"
+#include "Logging/LogMacros.h"
+#include "Logging/StructuredLog.h"
 #include "Misc/Paths.h"
 #include "Misc/ScopeLock.h"
 #include "Stats/StatsMisc.h"
@@ -355,6 +358,21 @@ void UGameplayTagsManager::LoadGameplayTagTables(bool bAllowAsyncLoad)
 
 	GameplayTagTables.Empty();
 
+	// Because we may unload the GameplayTags ini file later, foce the construction of other CDOs which use that config file now so that they don't reload it later
+	(void)GetDefault<URestrictedGameplayTagsList>();
+
+#if !UE_BUILD_SHIPPING && !UE_BUILD_TEST && DO_ENSURE
+	// Look for any other classes with the gameplay tags ini file that will cause it to be reloaded later
+	for (TObjectIterator<UClass> ClassIt; ClassIt; ++ClassIt)
+	{
+		ensureMsgf(ClassIt->ClassConfigName != Default->GetClass()->ClassConfigName
+			|| *ClassIt == UGameplayTagsSettings::StaticClass() 
+			|| *ClassIt == URestrictedGameplayTagsList::StaticClass(), 
+			TEXT("Class %s is using the GameplayTags config file which UGameplayTagsManager will attempt to unload after building the gameplay tag tree. ")
+			TEXT("This may cause the GameplayTags config file to be reloaded, consuming additional memory."), *ClassIt->GetPathName() );
+	}
+#endif // !UE_BUILD_SHIPPING && !UE_BUILD_TEST && DO_ENSURE
+
 #if !WITH_EDITOR
 	// If we're a cooked build and in a safe spot, start an async load so we can pipeline it
 	if (bAllowAsyncLoad && !IsLoading() && Default->GameplayTagTableList.Num() > 0)
