Add scope to AddTagIniSearchPath

#Synced-CL 51256439

diff --git a/Engine/Source/Runtime/GameplayTags/Private/GameplayTagsManager.cpp b/Engine/Source/Runtime/GameplayTags/Private/GameplayTagsManager.cpp
--- a/Engine/Source/Runtime/GameplayTags/Private/GameplayTagsManager.cpp
+++ b/Engine/Source/Runtime/GameplayTags/Private/GameplayTagsManager.cpp
@@ -404,6 +404,7 @@ void UGameplayTagsManager::LoadGameplayTagTables(bool bAllowAsyncLoad)
 
 void UGameplayTagsManager::AddTagIniSearchPath(const FString& RootDir, const TSet<FString>* PluginConfigsCache)
 {
+	TRACE_CPUPROFILER_EVENT_SCOPE(UGameplayTagsManager::AddTagIniSearchPath);
 	SCOPE_SECONDS_ACCUMULATOR(STAT_GameplayTags_AddTagIniSearchPath);
 
 	FGameplayTagSearchPathInfo* PathInfo = RegisteredSearchPaths.Find(RootDir);
